Building on GameplayTagQuery for reusability w/ audio gameplay objects

diff --git a/Engine/Plugins/AudioGameplay/Source/AudioGameplay/AudioGameplay.Build.cs b/Engine/Plugins/AudioGameplay/Source/AudioGameplay/AudioGameplay.Build.cs
--- a/Engine/Plugins/AudioGameplay/Source/AudioGameplay/AudioGameplay.Build.cs
+++ b/Engine/Plugins/AudioGameplay/Source/AudioGameplay/AudioGameplay.Build.cs
@@ -24,6 +24,7 @@ public class AudioGameplay : ModuleRules
 			new string[]
 			{
 				"Core",
+				"GameplayTags",
 			}
 			);
 			
diff --git a/Engine/Plugins/AudioGameplay/Source/AudioGameplay/Private/AudioGameplayRequirements.cpp b/Engine/Plugins/AudioGameplay/Source/AudioGameplay/Private/AudioGameplayRequirements.cpp
new file mode 100644
--- /dev/null
+++ b/Engine/Plugins/AudioGameplay/Source/AudioGameplay/Private/AudioGameplayRequirements.cpp
@@ -0,0 +1,28 @@
+// Copyright Epic Games, Inc. All Rights Reserved.
+
+#include "AudioGameplayRequirements.h"
+
+UAudioRequirementPreset::UAudioRequirementPreset(const FObjectInitializer& ObjectInitializer)
+	: Super(ObjectInitializer)
+{
+}
+
+bool UAudioRequirementPreset::Matches(const FGameplayTagContainer& Tags) const
+{
+	return Query.IsEmpty() || Query.Matches(Tags);
+}
+
+bool FAudioGameplayRequirements::Matches(const FGameplayTagContainer& Tags) const
+{
+	if (Preset && !Preset->Matches(Tags))
+	{
+		return false;
+	}
+
+	if (!Custom.IsEmpty() && !Custom.Matches(Tags))
+	{
+		return false;
+	}
+
+	return true;
+}
\ No newline at end of file
diff --git a/Engine/Plugins/AudioGameplay/Source/AudioGameplay/Public/AudioGameplayRequirements.h b/Engine/Plugins/AudioGameplay/Source/AudioGameplay/Public/AudioGameplayRequirements.h
new file mode 100644
--- /dev/null
+++ b/Engine/Plugins/AudioGameplay/Source/AudioGameplay/Public/AudioGameplayRequirements.h
@@ -0,0 +1,47 @@
+// Copyright Epic Games, Inc. All Rights Reserved.
+
+#pragma once
+
+#include "Engine/DataAsset.h"
+#include "GameplayTagContainer.h"
+
+#include "AudioGameplayRequirements.generated.h"
+
+/**
+ * UAudioRequirementPreset - A reusable GameplayTagQuery for audio
+ */
+UCLASS()
+class AUDIOGAMEPLAY_API UAudioRequirementPreset : public UDataAsset
+{
+	GENERATED_UCLASS_BODY()
+
+public:
+
+	~UAudioRequirementPreset() = default;
+
+	UPROPERTY(EditAnywhere, Category = "Gameplay Tag Query")
+	FGameplayTagQuery Query;
+
+	/** Returns true if the query is not empty and matches the provided tags */
+	bool Matches(const FGameplayTagContainer& Tags) const;
+};
+
+/**
+ * AudioGameplayRequirements - A set of requirements for audio gameplay features.
+ */
+USTRUCT()
+struct AUDIOGAMEPLAY_API FAudioGameplayRequirements
+{
+	GENERATED_USTRUCT_BODY()
+
+	/** An optional requirement preset tested against a collection of tags to determine a match.  Ignored if null. */
+	UPROPERTY(EditAnywhere, Category = "Requirements")
+	TObjectPtr<UAudioRequirementPreset> Preset = nullptr;
+
+	/** An optional custom query tested against a collection of tags to determine a match.  Ignored if empty. */
+	UPROPERTY(EditAnywhere, Category = "Requirements")
+	FGameplayTagQuery Custom;
+
+	/** Returns true if the preset and the custom query match the provided tags; ignores preset or custom query if they are empty. */
+	bool Matches(const FGameplayTagContainer& Tags) const;
+};
\ No newline at end of file
