- Adding some UPROP markup to SoundLibrary types for scripting support
- Add SoundLibAssetAction EUB and new action that converts existing BP libs to proper Asset libs

diff --git a/Engine/Plugins/AudioGameplay/Source/AudioGameplay/Public/AudioGameplayRequirements.h b/Engine/Plugins/AudioGameplay/Source/AudioGameplay/Public/AudioGameplayRequirements.h
--- a/Engine/Plugins/AudioGameplay/Source/AudioGameplay/Public/AudioGameplayRequirements.h
+++ b/Engine/Plugins/AudioGameplay/Source/AudioGameplay/Public/AudioGameplayRequirements.h
@@ -10,7 +10,7 @@
 /**
  * UAudioRequirementPreset - A reusable GameplayTagQuery for audio
  */
-UCLASS()
+UCLASS(BlueprintType)
 class AUDIOGAMEPLAY_API UAudioRequirementPreset : public UDataAsset
 {
 	GENERATED_UCLASS_BODY()
@@ -19,7 +19,7 @@ public:
 
 	~UAudioRequirementPreset() = default;
 
-	UPROPERTY(EditAnywhere, Category = "Gameplay Tag Query")
+	UPROPERTY(EditAnywhere, BlueprintReadWrite, Category = "Gameplay Tag Query")
 	FGameplayTagQuery Query;
 
 	/** Returns true if the query is not empty and matches the provided tags */
@@ -29,17 +29,17 @@ public:
 /**
  * AudioGameplayRequirements - A set of requirements for audio gameplay features.
  */
-USTRUCT()
+USTRUCT(BlueprintType)
 struct AUDIOGAMEPLAY_API FAudioGameplayRequirements
 {
 	GENERATED_USTRUCT_BODY()
 
 	/** An optional requirement preset tested against a collection of tags to determine a match.  Ignored if null. */
-	UPROPERTY(EditAnywhere, Category = "Requirements")
+	UPROPERTY(EditAnywhere, BlueprintReadWrite, Category = "Requirements")
 	TObjectPtr<UAudioRequirementPreset> Preset = nullptr;
 
 	/** An optional custom query tested against a collection of tags to determine a match.  Ignored if empty. */
-	UPROPERTY(EditAnywhere, Category = "Requirements")
+	UPROPERTY(EditAnywhere, BlueprintReadWrite, Category = "Requirements")
 	FGameplayTagQuery Custom;
 
 	/** Returns true if the preset and the custom query match the provided tags; ignores preset or custom query if they are empty. */
