Add AudioParameterComponent

UActorSoundParameterInterface::Fill now gathers parameters from all Actors and Components in the attachment chain

- Component implements ActorSoundParameterInterface
- Can act as a data store for FAudioParameter and pass it along to any sounds that play from the Owner actor
- Real-time parameter updates will be sent to any currently playing attached AudioComponents within the chain

diff --git a/Engine/Plugins/AudioGameplay/Source/AudioGameplay/AudioGameplay.Build.cs b/Engine/Plugins/AudioGameplay/Source/AudioGameplay/AudioGameplay.Build.cs
--- a/Engine/Plugins/AudioGameplay/Source/AudioGameplay/AudioGameplay.Build.cs
+++ b/Engine/Plugins/AudioGameplay/Source/AudioGameplay/AudioGameplay.Build.cs
@@ -12,6 +12,7 @@ public class AudioGameplay : ModuleRules
 			new string[]
 			{
 				"Core",
+				"AudioExtensions",
 				"GameplayTags",
 			}
 			);
diff --git a/Engine/Plugins/AudioGameplay/Source/AudioGameplay/Private/AudioParameterComponent.cpp b/Engine/Plugins/AudioGameplay/Source/AudioGameplay/Private/AudioParameterComponent.cpp
new file mode 100644
--- /dev/null
+++ b/Engine/Plugins/AudioGameplay/Source/AudioGameplay/Private/AudioParameterComponent.cpp
@@ -0,0 +1,173 @@
+// Copyright Epic Games, Inc. All Rights Reserved.
+#include "AudioParameterComponent.h"
+#include "Components/AudioComponent.h"
+#include "UObject/UObjectHash.h"
+
+
+void UAudioParameterComponent::GetActorSoundParams_Implementation(TArray<FAudioParameter>& Params) const
+{
+	Params.Append(GetParameters());
+}
+
+void UAudioParameterComponent::ResetParameters()
+{
+	int32 CurrSize = Parameters.Num();
+	Parameters.Reset(CurrSize);
+}
+
+void UAudioParameterComponent::SetTriggerParameter(FName InName)
+{
+	// todo@sweedin: use "GetAllAudioComponent" and execute this trigger on them
+	UE_LOG(LogAudio, Warning, TEXT("Trigger parameter not supported for UAudioParameterComponent"));
+}
+
+void UAudioParameterComponent::SetBoolParameter(FName InName, bool InValue)
+{
+	SetParameterInternal(FAudioParameter(InName, InValue));
+}
+
+void UAudioParameterComponent::SetBoolArrayParameter(FName InName, const TArray<bool>& InValue)
+{
+	SetParameterInternal(FAudioParameter(InName, InValue));
+}
+
+void UAudioParameterComponent::SetIntParameter(FName InName, int32 InValue)
+{
+	SetParameterInternal(FAudioParameter(InName, InValue));
+}
+
+void UAudioParameterComponent::SetIntArrayParameter(FName InName, const TArray<int32>& InValue)
+{
+	SetParameterInternal(FAudioParameter(InName, InValue));
+}
+
+void UAudioParameterComponent::SetFloatParameter(FName InName, float InValue)
+{
+	SetParameterInternal(FAudioParameter(InName, InValue));
+}
+
+void UAudioParameterComponent::SetFloatArrayParameter(FName InName, const TArray<float>& InValue)
+{
+	SetParameterInternal(FAudioParameter(InName, InValue));
+}
+
+void UAudioParameterComponent::SetStringParameter(FName InName, const FString& InValue)
+{
+	SetParameterInternal(FAudioParameter(InName, InValue));
+}
+
+void UAudioParameterComponent::SetStringArrayParameter(FName InName, const TArray<FString>& InValue)
+{
+	SetParameterInternal(FAudioParameter(InName, InValue));
+}
+
+void UAudioParameterComponent::SetObjectParameter(FName InName, UObject* InValue)
+{
+	SetParameterInternal(FAudioParameter(InName, InValue));
+}
+
+void UAudioParameterComponent::SetObjectArrayParameter(FName InName, const TArray<UObject*>& InValue)
+{
+	SetParameterInternal(FAudioParameter(InName, InValue));
+}
+
+void UAudioParameterComponent::SetParameters_Blueprint(const TArray<FAudioParameter>& InParameters)
+{
+	TArray<FAudioParameter> Values = InParameters;
+	SetParameters(MoveTemp(Values));
+}
+
+void UAudioParameterComponent::SetParameter(FAudioParameter&& InValue)
+{
+	SetParameterInternal(MoveTemp(InValue));
+}
+
+void UAudioParameterComponent::SetParameters(TArray<FAudioParameter>&& InValues)
+{
+	for (const FAudioParameter& Value : InValues)
+	{
+		if (FAudioParameter* CurrentParam = FAudioParameter::FindOrAddParam(Parameters, Value.ParamName))
+		{
+			constexpr bool bInTakeName = false;
+			CurrentParam->Merge(Value, bInTakeName);
+		}
+	}
+
+	// Forward to any AudioComponents currently playing on this actor (if any)
+	TArray<TObjectPtr<UAudioComponent>> Components;
+	GetAllAudioComponents(Components);
+
+	for (auto& Component : Components)
+	{
+		if (Component && Component->IsPlaying())
+		{
+			TArray<FAudioParameter> ValuesTemp = InValues;
+			Component->SetParameters(MoveTemp(ValuesTemp));
+		}
+	}
+}
+
+void UAudioParameterComponent::SetParameterInternal(FAudioParameter&& InParam)
+{
+	if (InParam.ParamName.IsNone())
+	{
+		return;
+	}
+
+	if (FAudioParameter* CurrentParam = FAudioParameter::FindOrAddParam(Parameters, InParam.ParamName))
+	{
+		constexpr bool bInTakeName = false;
+		CurrentParam->Merge(InParam, bInTakeName);
+	}
+
+	// Forward to any AudioComponents currently playing on this actor (if any)
+	TArray<UAudioComponent*> Components;
+	GetAllAudioComponents(Components);
+
+	for (auto& Component : Components)
+	{
+		if (Component && Component->IsPlaying())
+		{
+			Component->SetParameter(MoveTemp(InParam));
+		}
+	}
+}
+
+void UAudioParameterComponent::GetAllAudioComponents(TArray<UAudioComponent*>& Components) const
+{
+	AActor* RootActor = GetOwner();
+	if (USceneComponent* RootComp = RootActor->GetRootComponent())
+	{
+		// Walk up to the top-most attach actor in the hierarchy (will just be the RootActor if no attachment)
+		RootActor = RootComp->GetAttachmentRootActor();
+	}
+
+	// Helper to collect components from an actor
+	auto CollectFromActor = [&Components](const AActor* InActor)
+	{
+		if (InActor)
+		{
+			TArray<UAudioComponent*> TempComponents;
+			constexpr bool bIncludeChildActors = false;
+			InActor->GetComponents<UAudioComponent>(TempComponents, bIncludeChildActors);
+			Components.Append(MoveTemp(TempComponents));
+		}
+	};
+
+	if (RootActor)
+	{
+		// Grab any AudioComponents from our owner directly
+		CollectFromActor(RootActor);
+
+		// Recursively grab AudioComponents from our Attached Actor's
+		TArray<AActor*> AttachedActors;
+		constexpr bool bResetArray = false;
+		constexpr bool bRecursivelyIncludeAttachedActors = true;
+		RootActor->GetAttachedActors(AttachedActors, bResetArray, bRecursivelyIncludeAttachedActors);
+
+		for (AActor* Actor : AttachedActors)
+		{
+			CollectFromActor(Actor);
+		}
+	}
+}
diff --git a/Engine/Plugins/AudioGameplay/Source/AudioGameplay/Public/AudioParameterComponent.h b/Engine/Plugins/AudioGameplay/Source/AudioGameplay/Public/AudioParameterComponent.h
new file mode 100644
--- /dev/null
+++ b/Engine/Plugins/AudioGameplay/Source/AudioGameplay/Public/AudioParameterComponent.h
@@ -0,0 +1,64 @@
+// Copyright Epic Games, Inc. All Rights Reserved.
+#pragma once
+
+#include "Components/ActorComponent.h"
+#include "Containers/Array.h"
+#include "UObject/Object.h"
+#include "AudioParameter.h"
+#include "AudioParameterControllerInterface.h"
+#include "Audio/ActorSoundParameterInterface.h"
+
+#include "AudioParameterComponent.generated.h"
+
+
+class UAudioComponent;
+
+/**
+ *	UAudioParameterComponent - Can be used to set/store audio parameters and automatically dispatch them (through ActorSoundParameterInterface) 
+ *  to any sounds played by the component's Owner Actor
+ */
+UCLASS(BlueprintType, HideCategories=(Object, ActorComponent, Physics, Rendering, Mobility, LOD), meta=(BlueprintSpawnableComponent))
+class AUDIOGAMEPLAY_API UAudioParameterComponent : public UActorComponent, 
+		  										   public IAudioParameterControllerInterface, 
+												   public IActorSoundParameterInterface
+{
+	GENERATED_BODY()
+
+public:
+
+	UFUNCTION(BlueprintCallable, BlueprintPure, Category = Parameters)
+	const TArray<FAudioParameter>& GetParameters() const { return Parameters; }
+
+	/** Begin IActorSoundParameterInterface */
+	virtual void GetActorSoundParams_Implementation(TArray<FAudioParameter>& Params) const override;
+	/** End IActorSoundParameterInterface */
+
+	/** Begin IAudioParameterControllerInterface */
+	virtual void ResetParameters() override;
+	virtual void SetTriggerParameter(FName InName) override;
+	virtual void SetBoolParameter(FName InName, bool InValue) override;
+	virtual void SetBoolArrayParameter(FName InName, const TArray<bool>& InValue) override;
+	virtual void SetIntParameter(FName InName, int32 InInt) override;
+	virtual void SetIntArrayParameter(FName InName, const TArray<int32>& InValue) override;
+	virtual void SetFloatParameter(FName InName, float InValue) override;
+	virtual void SetFloatArrayParameter(FName InName, const TArray<float>& InValue) override;
+	virtual void SetStringParameter(FName InName, const FString& InValue) override;
+	virtual void SetStringArrayParameter(FName InName, const TArray<FString>& InValue) override;
+	virtual void SetObjectParameter(FName InName, UObject* InValue) override;
+	virtual void SetObjectArrayParameter(FName InName, const TArray<UObject*>& InValue) override;
+	virtual void SetParameters_Blueprint(const TArray<FAudioParameter>& InParameters) override;
+	virtual void SetParameter(FAudioParameter&& InValue) override;
+	virtual void SetParameters(TArray<FAudioParameter>&& InValues) override;
+	/** End IAudioParameterControllerInterface */
+
+private:
+
+	void SetParameterInternal(FAudioParameter&& InParam);
+	void GetAllAudioComponents(TArray<UAudioComponent*>& Components) const;
+
+	UPROPERTY(Transient)
+	TArray<TWeakObjectPtr<UAudioComponent>> ActiveComponents;
+
+	UPROPERTY(EditDefaultsOnly, Category = "Parameters")
+	TArray<FAudioParameter> Parameters;
+};
diff --git a/Engine/Source/Runtime/Engine/Classes/Audio/ActorSoundParameterInterface.h b/Engine/Source/Runtime/Engine/Classes/Audio/ActorSoundParameterInterface.h
--- a/Engine/Source/Runtime/Engine/Classes/Audio/ActorSoundParameterInterface.h
+++ b/Engine/Source/Runtime/Engine/Classes/Audio/ActorSoundParameterInterface.h
@@ -3,7 +3,6 @@
 
 #include "AudioParameterControllerInterface.h"
 #include "Containers/Array.h"
-#include "GameFramework/Actor.h"
 #include "UObject/Object.h"
 #include "UObject/Interface.h"
 
@@ -12,24 +11,30 @@
 
 // Forward Declarations
 class AActor;
+class UActorComponent;
 
 /** Interface used to allow an actor to automatically populate any sounds with parameters */
 class ENGINE_API IActorSoundParameterInterface : public IInterface
 {
-	GENERATED_IINTERFACE_BODY()
+	GENERATED_BODY()
 
 public:
-
-	UFUNCTION(BlueprintImplementableEvent)
+	// Overrides logic for gathering AudioParameters to set by default when an AudioComponent/ActiveSound plays with a given actor as its Owner.
+	UFUNCTION(BlueprintNativeEvent, Category = "Audio|Parameters", meta = (DisplayName = "Get Actor Audio Parameters"))
 	void GetActorSoundParams(TArray<FAudioParameter>& Params) const;
+	virtual void GetActorSoundParams_Implementation(TArray<FAudioParameter>& Params) const = 0;
 };
 
 UINTERFACE(BlueprintType)
 class ENGINE_API UActorSoundParameterInterface : public UInterface
 {
-	GENERATED_UINTERFACE_BODY()
+	GENERATED_BODY()
 
 public:
 
 	static void Fill(const AActor* OwningActor, TArray<FAudioParameter>& OutParams);
+
+private:
+
+	static void GetImplementers(const AActor* InActor, TArray<const AActor*>& OutActors, TArray<const UActorComponent*>& OutComponents);
 };
diff --git a/Engine/Source/Runtime/Engine/Private/Audio/ActorSoundParameterInterface.cpp b/Engine/Source/Runtime/Engine/Private/Audio/ActorSoundParameterInterface.cpp
--- a/Engine/Source/Runtime/Engine/Private/Audio/ActorSoundParameterInterface.cpp
+++ b/Engine/Source/Runtime/Engine/Private/Audio/ActorSoundParameterInterface.cpp
@@ -1,24 +1,73 @@
 // Copyright Epic Games, Inc. All Rights Reserved.
 #include "Audio/ActorSoundParameterInterface.h"
+#include "AudioParameter.h"
 #include "GameFramework/Actor.h"
 
 #include UE_INLINE_GENERATED_CPP_BY_NAME(ActorSoundParameterInterface)
 
 
-UActorSoundParameterInterface::UActorSoundParameterInterface(FObjectInitializer const& InObjectInitializer)
-	: Super(InObjectInitializer)
+void UActorSoundParameterInterface::Fill(const AActor* OwningActor, TArray<FAudioParameter>& OutParams)
 {
+	TArray<const AActor*> Actors;
+	TArray<const UActorComponent*> Components;
+	GetImplementers(OwningActor, Actors, Components);
+
+	TArray<FAudioParameter> TempParams;
+
+	for (const AActor* Actor : Actors)
+	{
+		IActorSoundParameterInterface::Execute_GetActorSoundParams(Actor, TempParams);
+		OutParams.Append(MoveTemp(TempParams));
+	}
+
+	for (const UActorComponent* Component : Components)
+	{
+		IActorSoundParameterInterface::Execute_GetActorSoundParams(Component, TempParams);
+		OutParams.Append(MoveTemp(TempParams));
+	}
 }
 
-void UActorSoundParameterInterface::Fill(const AActor* OwningActor, TArray<FAudioParameter>& OutParams)
+void UActorSoundParameterInterface::GetImplementers(const AActor* InActor, TArray<const AActor*>& OutActors, TArray<const UActorComponent*>& OutComponents)
 {
-	// Retrieve additional params from the actor if they implement a special interface
-	if (OwningActor && OwningActor->Implements<UActorSoundParameterInterface>())
+	if (!InActor)
 	{
-		TArray<FAudioParameter> ActorAudioParams;
-		IActorSoundParameterInterface::Execute_GetActorSoundParams(OwningActor, ActorAudioParams);
+		return;
+	}
 
-		OutParams.Append(ActorAudioParams);
+	// Helper to collect objects that implement this interface from an actor (and its components)
+	auto CollectFromActor = [&OutActors, &OutComponents](const AActor* InActor)
+	{
+		if (InActor)
+		{
+			if (InActor->Implements<UActorSoundParameterInterface>())
+			{
+				OutActors.Add(InActor);
+			}
+
+			TArray<UActorComponent*> Components = InActor->GetComponentsByInterface(UActorSoundParameterInterface::StaticClass());
+			OutComponents.Append(Components);
+		}
+	};
+
+	// Collect Actors/Components that implement this interface
+	const AActor* RootActor = InActor;
+	if (USceneComponent* RootComp = RootActor->GetRootComponent())
+	{
+		// Walk up to the top-most attach actor in the hierarchy (will just be the RootActor if no attachment)
+		RootActor = RootComp->GetAttachmentRootActor();
+	}
+
+	CollectFromActor(RootActor);
+
+	// Grab all attached actors (recursive)
+	TArray<AActor*> AttachedActors;
+	constexpr bool bResetArray = false;
+	constexpr bool bRecursivelyIncludeAttachedActors = true;
+	RootActor->GetAttachedActors(AttachedActors, bResetArray, bRecursivelyIncludeAttachedActors);
+
+	for (const AActor* Actor : AttachedActors)
+	{
+		CollectFromActor(Actor);
 	}
 }
 
