First pass at AudioComponentGroup

diff --git a/Engine/Plugins/AudioGameplay/Source/AudioGameplay/AudioGameplay.Build.cs b/Engine/Plugins/AudioGameplay/Source/AudioGameplay/AudioGameplay.Build.cs
--- a/Engine/Plugins/AudioGameplay/Source/AudioGameplay/AudioGameplay.Build.cs
+++ b/Engine/Plugins/AudioGameplay/Source/AudioGameplay/AudioGameplay.Build.cs
@@ -14,6 +14,7 @@ public class AudioGameplay : ModuleRules
 				"Core",
 				"AudioExtensions",
 				"GameplayTags",
+				"AudioExtensions",
 			}
 			);
 			
diff --git a/Engine/Plugins/AudioGameplay/Source/AudioGameplay/Private/AudioComponentGroup.cpp b/Engine/Plugins/AudioGameplay/Source/AudioGameplay/Private/AudioComponentGroup.cpp
new file mode 100644
--- /dev/null
+++ b/Engine/Plugins/AudioGameplay/Source/AudioGameplay/Private/AudioComponentGroup.cpp
@@ -0,0 +1,536 @@
+// Copyright Epic Games, Inc. All Rights Reserved.
+
+#include "AudioComponentGroup.h"
+#include "AudioParameter.h"
+#include "Kismet/GameplayStatics.h"
+
+#include "ProfilingDebugging/CpuProfilerTrace.h"
+
+UAudioComponentGroup::UAudioComponentGroup(const FObjectInitializer& ObjectInitializer)
+	: Super(ObjectInitializer)
+{
+	PrimaryComponentTick.bCanEverTick = true;
+	PrimaryComponentTick.bStartWithTickEnabled = true;
+	PrimaryComponentTick.bAllowTickOnDedicatedServer = false;
+	PrimaryComponentTick.TickGroup = ETickingGroup::TG_LastDemotable;
+}
+
+UAudioComponentGroup* UAudioComponentGroup::StaticGetOrCreateComponentGroup(AActor* Actor)
+{
+	TRACE_CPUPROFILER_EVENT_SCOPE(UAudioComponentGroup::StaticGetOrCreateComponentGroup);
+
+	AActor* Owner = Actor;
+	AActor* HighestValidOwner = nullptr;
+	TInlineComponentArray<UAudioComponentGroup*> SoundGroups;
+
+	while(Owner != nullptr)
+	{
+		Owner->GetComponents(SoundGroups);
+		if (!SoundGroups.IsEmpty())
+		{
+			return SoundGroups[0];
+		}
+
+		HighestValidOwner = Owner;
+		Owner = Owner->GetOwner();
+	}
+
+	if(HighestValidOwner != nullptr)
+	{
+		UAudioComponentGroup* SoundGroup = NewObject<UAudioComponentGroup>(HighestValidOwner);
+		SoundGroup->RegisterComponent();
+
+		return SoundGroup;
+	}
+	
+	return nullptr;
+}
+
+void UAudioComponentGroup::BeginPlay()
+{
+	Super::BeginPlay();
+}
+
+void UAudioComponentGroup::StopSound(USoundBase* Sound, const float FadeTime /*= 0.f*/)
+{
+	constexpr float TargetVolume = 0.f;
+
+	for (TObjectPtr<UAudioComponent>& Comp : Components)
+	{
+		if (Comp && Comp->IsPlaying() && Comp->Sound == Sound)
+		{
+			Comp->FadeOut(FadeTime, TargetVolume);
+		}
+	}
+
+	for (TWeakObjectPtr<UAudioComponent> Comp : ExternalComponents)
+	{
+		if (Comp.IsValid() && Comp->IsPlaying() && Comp->Sound == Sound)
+		{
+			Comp->FadeOut(FadeTime, TargetVolume);
+		}
+	}
+}
+
+bool UAudioComponentGroup::IsPlayingAny() const
+{
+	return OnStopped.IsBound();
+}
+
+void UAudioComponentGroup::BroadcastStopAll()
+{
+	OnStopped.Broadcast();
+}
+
+void UAudioComponentGroup::BroadcastKill()
+{
+	OnKilled.Broadcast();
+}
+
+void UAudioComponentGroup::BroadcastEvent(const FName EventName)
+{
+	if (TArray<FSoundCallback>* EventCallbacks = EventSubscriptions.Find(EventName))
+	{
+		for (auto EventIt = EventCallbacks->CreateIterator(); EventIt; ++EventIt)
+		{
+			if (EventIt->IsBound())
+			{
+				EventIt->Execute(EventName);
+			}
+			else
+			{
+				EventIt.RemoveCurrent();
+			}
+		}
+	}
+}
+
+UAudioComponent* UAudioComponentGroup::GetNextAvailableComponent()
+{
+	UAudioComponent* ReturnComponent = nullptr;
+
+	for (TObjectPtr<UAudioComponent>& Comp : Components)
+	{
+		if (Comp && !Comp->IsPlaying() && Comp->ActiveCount == 0)
+		{
+			ReturnComponent = ResetComponent(Comp);
+			break;
+		}
+	}
+
+	if (ReturnComponent == nullptr)
+	{
+		ReturnComponent = AddComponent();
+	}
+
+	for (TScriptInterface<IAudioComponentGroupExtension> Extension : Extensions)
+	{
+		Extension->OnComponentAdded(ReturnComponent);
+	}
+
+	return ReturnComponent;
+}
+
+UAudioComponent* UAudioComponentGroup::AddComponent()
+{
+	UAudioComponent* NewComponent = nullptr;
+	if (AActor* Owner = GetOwner())
+	{
+		NewComponent = NewObject<UAudioComponent>(Owner);
+
+		NewComponent->AttachToComponent(Owner->GetRootComponent(), FAttachmentTransformRules::SnapToTargetNotIncludingScale);
+		NewComponent->RegisterComponent();
+
+		ApplyParams(NewComponent);
+
+		FAudioComponentModifier Modifier = CachedModifier;
+		Modifier.Volume = GetComponentVolume();
+		ApplyModifiers(NewComponent, Modifier);
+
+		Components.Add(NewComponent);
+	}
+
+	return NewComponent;
+}
+
+UAudioComponent* UAudioComponentGroup::ResetComponent(UAudioComponent* Component) const
+{
+	if (Component)
+	{
+		Component->SetRelativeLocationAndRotation(FVector::ZeroVector, FRotator::ZeroRotator);
+		Component->ResetParameters();
+		Component->bAllowSpatialization = true;
+
+		ApplyParams(Component);
+
+		FAudioComponentModifier Modifier = CachedModifier;
+		Modifier.Volume = GetComponentVolume();
+		ApplyModifiers(Component, Modifier);
+	}
+
+	return Component;
+}
+
+void UAudioComponentGroup::RemoveComponent(const UAudioComponent* InComponent)
+{
+	int32 Index = Components.IndexOfByKey(InComponent);
+	if (Index != INDEX_NONE)
+	{
+		if (UAudioComponent* Component = Components[Index])
+		{
+			Component->Stop();
+		}
+		Components.RemoveAt(Index);
+	}
+}
+
+void UAudioComponentGroup::AddExternalComponent(UAudioComponent* ComponentToAdd)
+{
+	if (ComponentToAdd != nullptr)
+	{
+		ApplyParams(ComponentToAdd);
+
+		FAudioComponentModifier Modifier = CachedModifier;
+		Modifier.Volume = GetComponentVolume();
+		ApplyModifiers(ComponentToAdd, Modifier);
+
+		ExternalComponents.Add(TWeakObjectPtr<UAudioComponent>(ComponentToAdd));
+	}
+}
+
+void UAudioComponentGroup::EnableVirtualization()
+{
+	if (!bIsVirtualized)
+	{
+		bIsVirtualized = true;
+		OnVirtualized.Broadcast();
+	}
+}
+
+void UAudioComponentGroup::DisableVirtualization()
+{
+	if (bIsVirtualized)
+	{
+		bIsVirtualized = false;
+		OnUnvirtualized.Broadcast();
+	}
+}
+
+void UAudioComponentGroup::AddExtension(TScriptInterface<IAudioComponentGroupExtension> NewExtension)
+{
+	Extensions.AddUnique(NewExtension);
+	NewExtension->OnAddedToGroup(this);
+}
+
+void UAudioComponentGroup::RemoveExtension(TScriptInterface<IAudioComponentGroupExtension> NewExtension)
+{
+	Extensions.Remove(NewExtension);
+}
+
+void UAudioComponentGroup::UpdateExtensions(const float DeltaTime)
+{
+	FAudioComponentModifier ExtensionModifier;
+	for(TScriptInterface<IAudioComponentGroupExtension> Extension : Extensions)
+	{
+		Extension->Update(DeltaTime, this, ExtensionModifier);
+	}
+
+	ExtensionModifier.Combine(GroupModifier);
+
+	if(bIsVirtualized)
+	{
+		ExtensionModifier.Volume = 0.f;
+	}
+
+	if(ExtensionModifier.IsNearlyEqual(CachedModifier) == false)
+	{
+		IterateComponents([this, &ExtensionModifier](UAudioComponent* Component)
+		{
+			ApplyModifiers(Component, ExtensionModifier);
+		});
+	
+		CachedModifier = ExtensionModifier;
+	}
+}
+
+void UAudioComponentGroup::TickComponent(float DeltaTime, enum ELevelTick TickType, FActorComponentTickFunction* ThisTickFunction)
+{
+	TRACE_CPUPROFILER_EVENT_SCOPE(UAudioComponentGroup::TickComponent);
+
+	Super::TickComponent(DeltaTime, TickType, ThisTickFunction);
+
+	UpdateExtensions(DeltaTime);
+
+	UpdateComponentParameters();
+
+#if !(UE_BUILD_SHIPPING || UE_BUILD_TEST)
+	//AudioComponentGroupDebug::DebugPrint(this);
+#endif
+}
+
+float UAudioComponentGroup::GetFloatParamValue(const FName ParamName) const
+{
+	if (const FAudioParameter* FoundParam = GetParamInternal(ParamName))
+	{
+		return FoundParam->FloatParam;
+	}
+
+	return 0.f;
+}
+
+bool UAudioComponentGroup::GetBoolParamValue(const FName ParamName) const
+{
+	if (const FAudioParameter* FoundParam = GetParamInternal(ParamName))
+	{
+		return FoundParam->BoolParam;
+	}
+
+	return false;
+}
+
+FString UAudioComponentGroup::GetStringParamValue(const FName ParamName) const
+{
+	if (const FAudioParameter* FoundParam = GetParamInternal(ParamName))
+	{
+		return FoundParam->StringParam;
+	}
+
+	return FString();
+}
+
+void UAudioComponentGroup::SubscribeToStringParam(const FName ParamName, FStringParamCallback Delegate)
+{
+	if (TArray<FStringParamCallback>* StringCallbacks = StringSubscriptions.Find(ParamName))
+	{
+		StringCallbacks->Add(Delegate);
+	}
+	else
+	{
+		TArray<FStringParamCallback> NewSubscription = { Delegate };
+		StringSubscriptions.Add(ParamName, NewSubscription);
+	}
+}
+
+void UAudioComponentGroup::SubscribeToEvent(const FName EventName, FSoundCallback Delegate)
+{
+	if (TArray<FSoundCallback>* EventCallbacks = EventSubscriptions.Find(EventName))
+	{
+		EventCallbacks->AddUnique(Delegate);
+	}
+	else
+	{
+		TArray<FSoundCallback> NewSubscription = { Delegate };
+		EventSubscriptions.Add(EventName, NewSubscription);
+	}
+}
+
+void UAudioComponentGroup::SubscribeToBool(const FName ParamName, FBoolParamCallback Delegate)
+{
+	if (TArray<FBoolParamCallback>* Callbacks = BoolSubscriptions.Find(ParamName))
+	{
+		Callbacks->AddUnique(Delegate);
+	}
+	else
+	{
+		TArray<FBoolParamCallback> NewSubscription = { Delegate };
+		BoolSubscriptions.Add(ParamName, NewSubscription);
+	}
+}
+
+void UAudioComponentGroup::ResetParameters()
+{
+	PersistentParams.Reset();
+}
+
+void UAudioComponentGroup::SetTriggerParameter(FName InName)
+{
+	SetParameters({ FAudioParameter(InName, true) });
+}
+
+void UAudioComponentGroup::SetBoolParameter(FName InName, bool InBool)
+{
+	SetParameters({ FAudioParameter(InName, InBool) });
+
+	if (!InName.IsNone())
+	{
+		if (TArray<FBoolParamCallback>* BoolCallbacks = BoolSubscriptions.Find(InName))
+		{
+			for (TArray<FBoolParamCallback>::TIterator BoolCallbackIt = BoolCallbacks->CreateIterator(); BoolCallbackIt; ++BoolCallbackIt)
+			{
+				if (BoolCallbackIt->IsBound())
+				{
+					BoolCallbackIt->Execute(InBool);
+				}
+				else
+				{
+					BoolCallbackIt.RemoveCurrent();
+				}
+			}
+		}
+	}
+}
+
+void UAudioComponentGroup::SetBoolArrayParameter(FName InName, const TArray<bool>& InValue)
+{
+	SetParameters({ FAudioParameter(InName, InValue) });
+}
+
+void UAudioComponentGroup::SetIntParameter(FName InName, int32 InInt)
+{
+	SetParameters({ FAudioParameter(InName, InInt) });
+}
+
+void UAudioComponentGroup::SetIntArrayParameter(FName InName, const TArray<int32>& InValue)
+{
+	SetParameters({ FAudioParameter(InName, InValue) });
+}
+
+void UAudioComponentGroup::SetFloatParameter(FName InName, float InFloat)
+{
+	SetParameters({ FAudioParameter(InName, InFloat) });
+}
+
+void UAudioComponentGroup::SetFloatArrayParameter(FName InName, const TArray<float>& InValue)
+{
+	SetParameters({ FAudioParameter(InName, InValue) });
+}
+
+void UAudioComponentGroup::SetStringParameter(FName InName, const FString& InValue)
+{
+	SetParameters({ FAudioParameter(InName, InValue) });
+
+	if (TArray<FStringParamCallback>* StringCallbacks = StringSubscriptions.Find(InName))
+	{
+		for (TArray<FStringParamCallback>::TIterator StringCallbackIt = StringCallbacks->CreateIterator(); StringCallbacks; ++StringCallbacks)
+		{
+			if (StringCallbackIt->IsBound())
+			{
+				StringCallbackIt->Execute(InValue);
+			}
+			else
+			{
+				StringCallbackIt.RemoveCurrent();
+			}
+		}
+	}
+}
+
+void UAudioComponentGroup::SetStringArrayParameter(FName InName, const TArray<FString>& InValue)
+{
+	SetParameters({ FAudioParameter(InName, InValue) });
+}
+
+void UAudioComponentGroup::SetObjectParameter(FName InName, UObject* InValue)
+{
+	SetParameters({ FAudioParameter(InName, InValue) });
+}
+
+void UAudioComponentGroup::SetObjectArrayParameter(FName InName, const TArray<UObject*>& InValue)
+{
+	SetParameters({ FAudioParameter(InName, InValue) });
+}
+
+void UAudioComponentGroup::SetParameters_Blueprint(const TArray<FAudioParameter>& InParameters)
+{
+	TArray<FAudioParameter> Values = InParameters;
+	SetParameters(MoveTemp(Values));
+}
+
+void UAudioComponentGroup::SetParameter(FAudioParameter&& InValue)
+{
+	SetParameters({ MoveTemp(InValue) });
+}
+
+void UAudioComponentGroup::SetParameters(TArray<FAudioParameter>&& InValues)
+{
+	InValues.RemoveAll([](FAudioParameter& Param) { return Param.ParamName.IsNone(); });
+	FAudioParameter::Merge(MoveTemp(InValues), ParamsToSet);
+}
+
+void UAudioComponentGroup::ApplyParams(UAudioComponent* Component) const
+{
+	if (Component)
+	{
+		TRACE_CPUPROFILER_EVENT_SCOPE(UAudioComponentGroup::ApplyParams);
+
+		// if the component is playing already, InstanceParameters will be ignored
+		if (Component->IsPlaying())
+		{
+			TArray<FAudioParameter> Params = PersistentParams;
+			Component->SetParameters(MoveTemp(Params));
+		}
+		else
+		{
+			// if the component hasn't started playing yet, all of its instance parameters will be passed along at once when it starts
+			TArray<FAudioParameter>& InstanceParameters = Component->GetInstanceParameters();
+			InstanceParameters.Append(PersistentParams);
+		}
+	}
+}
+
+void UAudioComponentGroup::ApplyModifiers(UAudioComponent* Component, const FAudioComponentModifier& Modifier) const
+{
+	if (Component)
+	{
+		Component->SetVolumeMultiplier(Modifier.Volume);
+		Component->SetPitchMultiplier(Modifier.Pitch);
+		Component->SetLowPassFilterFrequency(Modifier.LowPassFrequency);
+	}
+}
+
+void UAudioComponentGroup::IterateComponents(const TFunction<void(UAudioComponent*)> OnIterate)
+{
+	Components.RemoveAll([](const UAudioComponent* Component) { return !Component; });
+	for (UAudioComponent* Component : Components)
+	{
+		if (Component->IsActive())
+		{
+			OnIterate(Component);
+		}
+	}
+
+	ExternalComponents.RemoveAll([](const TWeakObjectPtr<UAudioComponent> WeakComponent) { return !WeakComponent.IsValid(); });
+	for (TWeakObjectPtr<UAudioComponent>& WeakComponent : ExternalComponents)
+	{
+		if (WeakComponent->IsActive())
+		{
+			OnIterate(WeakComponent.Get());
+		}
+	}
+}
+
+void UAudioComponentGroup::UpdateComponentParameters()
+{
+	if(ParamsToSet.IsEmpty())
+	{
+		return;
+	}
+
+	TRACE_CPUPROFILER_EVENT_SCOPE(UAudioComponentGroup::UpdateComponentParameters);
+
+	IterateComponents([this](UAudioComponent* Component)
+	{
+		// todo: this could be a LOT of tiny allocs... can anything be done about that?
+		TArray<FAudioParameter> Values = ParamsToSet;
+		Component->SetParameters(MoveTemp(Values));
+	});
+	
+	TArray<FAudioParameter> Values = ParamsToSet;
+	FAudioParameter::Merge(MoveTemp(Values), PersistentParams);
+	ParamsToSet.Reset();
+}
+
+float UAudioComponentGroup::GetComponentVolume() const
+{
+	if (bIsVirtualized)
+	{
+		return 0.f;
+	}
+
+	return CachedModifier.Volume;
+}
+
+const FAudioParameter* UAudioComponentGroup::GetParamInternal(const FName ParamName) const
+{
+	return PersistentParams.FindByPredicate([ParamName](const FAudioParameter& Param) { return Param.ParamName == ParamName; });
+}
diff --git a/Engine/Plugins/AudioGameplay/Source/AudioGameplay/Public/AudioComponentGroup.h b/Engine/Plugins/AudioGameplay/Source/AudioGameplay/Public/AudioComponentGroup.h
new file mode 100644
--- /dev/null
+++ b/Engine/Plugins/AudioGameplay/Source/AudioGameplay/Public/AudioComponentGroup.h
@@ -0,0 +1,179 @@
+// Copyright Epic Games, Inc. All Rights Reserved.
+
+#pragma once
+
+#include "AudioComponentGroupExtension.h"
+#include "AudioParameterControllerInterface.h"
+#include "Components/AudioComponent.h"
+#include "AudioComponentGroup.generated.h"
+
+class USoundBase;
+class UParamCollection;
+
+DECLARE_DYNAMIC_MULTICAST_DELEGATE(FSoundGroupChanged);
+DECLARE_DYNAMIC_DELEGATE_OneParam(FSoundCallback, const FName&, EventName);
+DECLARE_DYNAMIC_DELEGATE_OneParam(FBoolParamCallback, const bool, ParamValue);
+DECLARE_DYNAMIC_DELEGATE_OneParam(FStringParamCallback, const FString&, Value);
+
+/*
+ * Automatic Handler for voices and parameters across any number of AudioComponents
+ */
+UCLASS(BlueprintType, DefaultToInstanced)
+class AUDIOGAMEPLAY_API UAudioComponentGroup : public USceneComponent, public IAudioParameterControllerInterface
+{
+	GENERATED_BODY()
+
+public:
+
+	UAudioComponentGroup(const FObjectInitializer& ObjectInitializer);
+
+	virtual ~UAudioComponentGroup() = default;
+
+	UFUNCTION(BlueprintCallable, BlueprintCosmetic, Category = AudioComponentGroup)
+	static UAudioComponentGroup* StaticGetOrCreateComponentGroup(AActor* Actor);
+
+	virtual void BeginPlay() override;
+
+	// Stop all instances of this Sound on any internal or external components
+	UFUNCTION(BlueprintCallable, BlueprintCosmetic, Category = AudioComponentGroup)
+	void StopSound(USoundBase* Sound, const float FadeTime = 0.f);
+
+	UFUNCTION(BlueprintPure, BlueprintCosmetic, Category = AudioComponentGroup)
+	bool IsPlayingAny() const;
+
+	UFUNCTION(BlueprintPure, BlueprintCosmetic, Category = AudioComponentGroup)
+	bool IsVirtualized() const { return bIsVirtualized; }
+
+	UFUNCTION(BlueprintCallable, BlueprintCosmetic, Category = AudioComponentGroup)
+	void BroadcastStopAll();
+
+	UFUNCTION(BlueprintCallable, BlueprintCosmetic, Category = AudioComponentGroup)
+	void BroadcastKill();
+
+	UFUNCTION(BlueprintCallable, BlueprintCosmetic, Category = AudioComponentGroup)
+	void BroadcastEvent(const FName EventName);
+
+	// Component interaction
+	UAudioComponent* GetNextAvailableComponent();
+
+	UAudioComponent* AddComponent();
+	UAudioComponent* ResetComponent(UAudioComponent* Component) const;
+	void RemoveComponent(const UAudioComponent* InComponent);
+
+	//Allow an externally created AudioComponent to share parameters with this SoundGroup
+	UFUNCTION(BlueprintCallable, BlueprintCosmetic, Category = AudioComponentGroup)
+	void AddExternalComponent(UAudioComponent* ComponentToAdd);
+
+	UFUNCTION(BlueprintCallable, BlueprintCosmetic, Category = AudioComponentGroup)
+	void EnableVirtualization();
+
+	UFUNCTION(BlueprintCallable, BlueprintCosmetic, Category = AudioComponentGroup)
+	void DisableVirtualization();
+
+	UFUNCTION(BlueprintCallable, BlueprintCosmetic, Category = AudioComponentGroup)
+	void SetVolumeMultiplier(const float InVolume) { GroupModifier.Volume = InVolume; }
+	
+	UFUNCTION(BlueprintCallable, BlueprintCosmetic, Category = AudioComponentGroup)
+	void SetPitchMultiplier(const float InPitch) { GroupModifier.Pitch = InPitch; }
+
+	UFUNCTION(BlueprintCallable, BlueprintCosmetic, Category = AudioComponentGroup)
+	void SetLowPassFilter(const float InFrequency) { GroupModifier.LowPassFrequency = InFrequency; }
+
+	const TArray<FAudioParameter>& GetParams() { return PersistentParams; }
+
+	UFUNCTION(BlueprintCallable, BlueprintCosmetic, Category = AudioComponentGroup)
+	void AddExtension(TScriptInterface<IAudioComponentGroupExtension> NewExtension);
+	
+	UFUNCTION(BlueprintCallable, BlueprintCosmetic, Category = AudioComponentGroup)
+	void RemoveExtension(TScriptInterface<IAudioComponentGroupExtension> NewExtension);
+
+	void UpdateExtensions(const float DeltaTime);
+
+	virtual void TickComponent(float DeltaTime, enum ELevelTick TickType, FActorComponentTickFunction *ThisTickFunction) override;
+
+	UFUNCTION(BlueprintPure, BlueprintCosmetic, Category = AudioComponentGroup)
+	float GetFloatParamValue(const FName ParamName) const;
+
+	UFUNCTION(BlueprintPure, BlueprintCosmetic, Category = AudioComponentGroup)
+	bool GetBoolParamValue(const FName ParamName) const;
+	 
+	UFUNCTION(BlueprintPure, BlueprintCosmetic, Category = AudioComponentGroup)
+	FString GetStringParamValue(const FName ParamName) const;
+
+	void SubscribeToStringParam(const FName ParamName, FStringParamCallback Delegate);
+	void SubscribeToEvent(const FName EventName, FSoundCallback Delegate);
+	void SubscribeToBool(const FName ParamName, FBoolParamCallback Delegate);
+
+	//~ Begin IAudioParameterControllerInterface interface 
+	virtual void ResetParameters() override;
+
+	virtual void SetTriggerParameter(FName InName) override;
+	virtual void SetBoolParameter(FName InName, bool InBool) override;
+	virtual void SetBoolArrayParameter(FName InName, const TArray<bool>& InValue) override;
+	virtual void SetIntParameter(FName InName, int32 InInt) override;
+	virtual void SetIntArrayParameter(FName InName, const TArray<int32>& InValue) override;
+	virtual void SetFloatParameter(FName InName, float InFloat) override;
+	virtual void SetFloatArrayParameter(FName InName, const TArray<float>& InValue) override;
+	virtual void SetStringParameter(FName InName, const FString& InValue) override;
+	virtual void SetStringArrayParameter(FName InName, const TArray<FString>& InValue) override;
+	virtual void SetObjectParameter(FName InName, UObject* InValue) override;
+	virtual void SetObjectArrayParameter(FName InName, const TArray<UObject*>& InValue) override;
+
+	virtual void SetParameter(FAudioParameter&& InValue) override;
+	virtual void SetParameters(TArray<FAudioParameter>&& InValues) override;
+	virtual void SetParameters_Blueprint(const TArray<FAudioParameter>& InParameters) override;
+	//~ End IAudioParameterControllerInterface interface
+
+	UPROPERTY(BlueprintAssignable)
+	FSoundGroupChanged OnStopped;
+
+	UPROPERTY(BlueprintAssignable)
+	FSoundGroupChanged OnKilled;
+
+	UPROPERTY(BlueprintAssignable)
+	FSoundGroupChanged OnVirtualized;
+
+	UPROPERTY(BlueprintAssignable)
+	FSoundGroupChanged OnUnvirtualized;
+
+protected:
+
+	void ApplyParams(UAudioComponent* Component) const;
+	void ApplyModifiers(UAudioComponent* Component, const FAudioComponentModifier& Modifier) const;
+	void IterateComponents(const TFunction<void(UAudioComponent*)> OnIterate);
+
+	void UpdateComponentParameters();
+
+	float GetComponentVolume() const;
+
+	const FAudioParameter* GetParamInternal(const FName ParamName) const;
+
+	UPROPERTY(VisibleAnywhere, Category = AudioComponentGroup)
+	TArray<TObjectPtr<UAudioComponent>> Components;
+
+	UPROPERTY(Transient)
+	TArray<FAudioParameter> ParamsToSet;
+	
+	UPROPERTY(VisibleAnywhere, BlueprintReadOnly, Category = AudioComponentGroup)
+	TArray<FAudioParameter> PersistentParams;
+
+	UPROPERTY(Transient)
+	TArray<TScriptInterface<IAudioComponentGroupExtension>> Extensions;
+
+	// Modifier set externally via BP functions
+	FAudioComponentModifier GroupModifier;
+
+	// final values set last update
+	FAudioComponentModifier CachedModifier;
+
+	//Components managed externally that won't be used in the pool, but can still share parameters
+	TArray<TWeakObjectPtr<UAudioComponent>> ExternalComponents;
+
+	TMap<FName, TArray<FStringParamCallback>> StringSubscriptions;
+	TMap<FName, TArray<FSoundCallback>> EventSubscriptions;
+	TMap<FName, TArray<FBoolParamCallback>> BoolSubscriptions;
+
+	bool bIsVirtualized = false;
+
+	friend class AudioComponentGroupDebug;
+};
diff --git a/Engine/Plugins/AudioGameplay/Source/AudioGameplay/Public/AudioComponentGroupExtension.h b/Engine/Plugins/AudioGameplay/Source/AudioGameplay/Public/AudioComponentGroupExtension.h
new file mode 100644
--- /dev/null
+++ b/Engine/Plugins/AudioGameplay/Source/AudioGameplay/Public/AudioComponentGroupExtension.h
@@ -0,0 +1,56 @@
+// Copyright Epic Games, Inc. All Rights Reserved.
+
+#pragma once
+
+#include "AudioDefines.h"
+#include "UObject/Interface.h"
+#include "AudioComponentGroupExtension.generated.h"
+
+class UAudioComponent;
+class UAudioComponentGroup;
+
+struct FAudioComponentModifier
+{
+	float Volume = 1.f;
+	float Pitch = 1.f;
+	float LowPassFrequency = MAX_FILTER_FREQUENCY;
+
+	void Combine(const FAudioComponentModifier& Other)
+	{
+		Volume *= Other.Volume;
+		Pitch *= Other.Pitch;
+		LowPassFrequency = FMath::Min(LowPassFrequency, Other.LowPassFrequency);
+	}
+
+	bool IsNearlyEqual(const FAudioComponentModifier& Other) const
+	{
+		return FMath::IsNearlyEqual(Volume, Other.Volume)
+			&& FMath::IsNearlyEqual(Pitch, Other.Pitch)
+			&& FMath::IsNearlyEqual(LowPassFrequency, Other.LowPassFrequency);
+	}
+
+	static FAudioComponentModifier& Default()
+	{
+		static FAudioComponentModifier DefaultModifier;
+		return DefaultModifier;
+	}
+};
+
+UINTERFACE(BlueprintType)
+class AUDIOGAMEPLAY_API UAudioComponentGroupExtension : public UInterface
+{
+	GENERATED_BODY()
+};
+
+class AUDIOGAMEPLAY_API IAudioComponentGroupExtension : public IInterface
+{
+	GENERATED_BODY()
+
+public:
+
+	virtual void Update(const float DeltaTime, UAudioComponentGroup* Group, FAudioComponentModifier& OutModifier) {}
+
+	virtual void OnAddedToGroup(UAudioComponentGroup* NewGroup) {}
+
+	virtual void OnComponentAdded(UAudioComponent* NewComponent) {}
+};
\ No newline at end of file
