AudioParameterComponent

- Fix up case where we need to send FAudioParameter updates to currently playing sounds w/o using InstanceParameter array, since that was being prioritized and stomping FAudioParameter's coming in from the ActorSoundParameterInterface

diff --git a/Engine/Plugins/AudioGameplay/Source/AudioGameplay/Private/AudioParameterComponent.cpp b/Engine/Plugins/AudioGameplay/Source/AudioGameplay/Private/AudioParameterComponent.cpp
--- a/Engine/Plugins/AudioGameplay/Source/AudioGameplay/Private/AudioParameterComponent.cpp
+++ b/Engine/Plugins/AudioGameplay/Source/AudioGameplay/Private/AudioParameterComponent.cpp
@@ -4,8 +4,13 @@
 #include "Components/AudioComponent.h"
 #include "GameFramework/Actor.h"
 #include "UObject/UObjectHash.h"
+#include "AudioDevice.h"
+#include "ActiveSound.h"
+#include "Kismet/KismetStringLibrary.h"
 
 
+DEFINE_LOG_CATEGORY(LogAudioParameterComponent);
+
 void UAudioParameterComponent::GetActorSoundParams_Implementation(TArray<FAudioParameter>& Params) const
 {
 	Params.Append(GetParameters());
@@ -122,15 +127,47 @@ void UAudioParameterComponent::SetParameterInternal(FAudioParameter&& InParam)
 		CurrentParam->Merge(InParam, bInTakeName);
 	}
 
+	// Optional logging
+	LogParameter(InParam);
+
 	// Forward to any AudioComponents currently playing on this actor (if any)
 	TArray<UAudioComponent*> Components;
 	GetAllAudioComponents(Components);
 
 	for (auto& Component : Components)
 	{
-		if (Component && Component->IsPlaying())
+		if (Component == nullptr)
 		{
-			Component->SetParameter(MoveTemp(InParam));
+			continue;
+		}
+
+		if (Component->IsPlaying() && !Component->GetDisableParameterUpdatesWhilePlaying())
+		{
+			if (FAudioDevice* AudioDevice = Component->GetAudioDevice())
+			{
+				if (USoundBase* Sound = Component->GetSound())
+				{
+					static const FName ProxyFeatureName("AudioParameterComponent");
+					Sound->InitParameters(Parameters, ProxyFeatureName);
+				}
+
+				// Prior call to InitParameters can prune parameters if they are
+				// invalid, so check here to avoid unnecessary pass of empty array.
+				if (!Parameters.IsEmpty())
+				{
+					DECLARE_CYCLE_STAT(TEXT("FAudioThreadTask.SoundParameterControllerInterface.SetParameters"), STAT_AudioSetParameters, STATGROUP_AudioThreadCommands);
+
+					TArray<FAudioParameter> ParametersCopy = Parameters;
+					AudioDevice->SendCommandToActiveSounds(Component->GetInstanceOwnerID(), [AudioDevice, Params = MoveTemp(ParametersCopy)](FActiveSound& ActiveSound)
+					{
+						if (Audio::IParameterTransmitter* Transmitter = ActiveSound.GetTransmitter())
+						{
+							TArray<FAudioParameter> TempParams = Params;
+							Transmitter->SetParameters(MoveTemp(TempParams));
+						}
+					}, GET_STATID(STAT_AudioSetParameters));
+				}
+			}
 		}
 	}
 }
@@ -173,3 +210,42 @@ void UAudioParameterComponent::GetAllAudioComponents(TArray<UAudioComponent*>& C
 		}
 	}
 }
+
+void UAudioParameterComponent::LogParameter(FAudioParameter& InParam)
+{
+	if (!UE_LOG_ACTIVE(LogAudioParameterComponent, VeryVerbose))
+	{
+		return;
+	}
+
+	const AActor* Owner = GetOwner();
+	if (Owner == nullptr)
+	{
+		return;
+	}
+
+	FString ParamTypeString;
+	const UEnum* ParamTypeEnum = StaticEnum<EAudioParameterType>();
+	if (ParamTypeEnum != nullptr)
+	{
+		ParamTypeString = ParamTypeEnum->GetNameStringByIndex(static_cast<int32>(InParam.ParamType));
+	}
+
+	FString ParamValueString;
+	switch (InParam.ParamType)
+	{
+		case EAudioParameterType::Boolean: 
+			ParamValueString = UKismetStringLibrary::Conv_BoolToString(InParam.BoolParam);
+			break;
+
+		case EAudioParameterType::Integer:
+			ParamValueString = UKismetStringLibrary::Conv_IntToString(InParam.IntParam);
+			break;
+
+		case EAudioParameterType::Float:
+			ParamValueString = FString::SanitizeFloat(InParam.FloatParam);
+			break;
+	}
+
+	UE_LOG(LogAudioParameterComponent, VeryVerbose, TEXT("%s: Set %s %s to %s"), *Owner->GetName(), *ParamTypeString, *InParam.ParamName.ToString(), *ParamValueString)
+}
diff --git a/Engine/Plugins/AudioGameplay/Source/AudioGameplay/Public/AudioParameterComponent.h b/Engine/Plugins/AudioGameplay/Source/AudioGameplay/Public/AudioParameterComponent.h
--- a/Engine/Plugins/AudioGameplay/Source/AudioGameplay/Public/AudioParameterComponent.h
+++ b/Engine/Plugins/AudioGameplay/Source/AudioGameplay/Public/AudioParameterComponent.h
@@ -11,6 +11,8 @@
 #include "AudioParameterComponent.generated.h"
 
 
+DECLARE_LOG_CATEGORY_EXTERN(LogAudioParameterComponent, Log, All);
+
 class UAudioComponent;
 
 /**
@@ -55,6 +57,7 @@ private:
 
 	void SetParameterInternal(FAudioParameter&& InParam);
 	void GetAllAudioComponents(TArray<UAudioComponent*>& Components) const;
+	void LogParameter(FAudioParameter& InParam);
 
 	UPROPERTY(Transient)
 	TArray<TWeakObjectPtr<UAudioComponent>> ActiveComponents;
