don't hold on to trigger parameters in AudioComponentGroup

diff --git a/Engine/Plugins/AudioGameplay/Source/AudioGameplay/Private/AudioComponentGroup.cpp b/Engine/Plugins/AudioGameplay/Source/AudioGameplay/Private/AudioComponentGroup.cpp
--- a/Engine/Plugins/AudioGameplay/Source/AudioGameplay/Private/AudioComponentGroup.cpp
+++ b/Engine/Plugins/AudioGameplay/Source/AudioGameplay/Private/AudioComponentGroup.cpp
@@ -347,7 +347,13 @@ void UAudioComponentGroup::ResetParameters()
 
 void UAudioComponentGroup::SetTriggerParameter(FName InName)
 {
-	SetParameters({ FAudioParameter(InName, true) });
+	IterateComponents([InName](UAudioComponent* Component)
+	{
+		if(Component->IsPlaying())
+		{
+			Component->SetTriggerParameter(InName);
+		}
+	});
 }
 
 void UAudioComponentGroup::SetBoolParameter(FName InName, bool InBool)
