expose audio component group parameter subscription funcs to blueprint, move delegate broadcasting to SetParameters to ensure that parameters set from anywhere get broadcast as expected

diff --git a/Engine/Plugins/AudioGameplay/Source/AudioGameplay/Private/AudioComponentGroup.cpp b/Engine/Plugins/AudioGameplay/Source/AudioGameplay/Private/AudioComponentGroup.cpp
--- a/Engine/Plugins/AudioGameplay/Source/AudioGameplay/Private/AudioComponentGroup.cpp
+++ b/Engine/Plugins/AudioGameplay/Source/AudioGameplay/Private/AudioComponentGroup.cpp
@@ -339,6 +339,42 @@ void UAudioComponentGroup::SubscribeToBool(const FName ParamName, FBoolParamCall
 	}
 }
 
+void UAudioComponentGroup::UnsubscribeObject(const UObject* Object)
+{
+	for (TPair<FName,TArray<FBoolParamCallback>>& BoolCallbacks : BoolSubscriptions)
+	{
+		for (TArray<FBoolParamCallback>::TIterator BoolCallbackIt = BoolCallbacks.Value.CreateIterator(); BoolCallbackIt; ++BoolCallbackIt)
+		{
+			if (!BoolCallbackIt->IsBound() || BoolCallbackIt->GetUObject() == Object)
+			{
+				BoolCallbackIt.RemoveCurrent();
+			}
+		}
+	}
+
+	for (TPair<FName,TArray<FStringParamCallback>>& StringCallbacks : StringSubscriptions)
+	{
+		for (TArray<FStringParamCallback>::TIterator StringCallbackIt = StringCallbacks.Value.CreateIterator(); StringCallbackIt; ++StringCallbackIt)
+		{
+			if (!StringCallbackIt->IsBound() || StringCallbackIt->GetUObject() == Object)
+			{
+				StringCallbackIt.RemoveCurrent();
+			}
+		}
+	}
+
+	for (TPair<FName,TArray<FSoundCallback>>& Callbacks : EventSubscriptions)
+	{
+		for (TArray<FSoundCallback>::TIterator CallbackIt = Callbacks.Value.CreateIterator(); CallbackIt; ++CallbackIt)
+		{
+			if (!CallbackIt->IsBound() || CallbackIt->GetUObject() == Object)
+			{
+				CallbackIt.RemoveCurrent();
+			}
+		}
+	}
+}
+
 void UAudioComponentGroup::ResetParameters()
 {
 	PersistentParams.Reset();
@@ -346,9 +382,11 @@ void UAudioComponentGroup::ResetParameters()
 
 void UAudioComponentGroup::SetTriggerParameter(FName InName)
 {
+	ExecuteEventSubscriptions(InName);
+
 	IterateComponents([InName](UAudioComponent* Component)
 	{
-		if(Component->IsPlaying())
+		if (Component->IsPlaying())
 		{
 			Component->SetTriggerParameter(InName);
 		}
@@ -358,24 +396,6 @@ void UAudioComponentGroup::SetTriggerParameter(FName InName)
 void UAudioComponentGroup::SetBoolParameter(FName InName, bool InBool)
 {
 	SetParameters({ FAudioParameter(InName, InBool) });
-
-	if (!InName.IsNone())
-	{
-		if (TArray<FBoolParamCallback>* BoolCallbacks = BoolSubscriptions.Find(InName))
-		{
-			for (TArray<FBoolParamCallback>::TIterator BoolCallbackIt = BoolCallbacks->CreateIterator(); BoolCallbackIt; ++BoolCallbackIt)
-			{
-				if (BoolCallbackIt->IsBound())
-				{
-					BoolCallbackIt->Execute(InBool);
-				}
-				else
-				{
-					BoolCallbackIt.RemoveCurrent();
-				}
-			}
-		}
-	}
 }
 
 void UAudioComponentGroup::SetBoolArrayParameter(FName InName, const TArray<bool>& InValue)
@@ -406,21 +426,6 @@ void UAudioComponentGroup::SetFloatArrayParameter(FName InName, const TArray<flo
 void UAudioComponentGroup::SetStringParameter(FName InName, const FString& InValue)
 {
 	SetParameters({ FAudioParameter(InName, InValue) });
-
-	if (TArray<FStringParamCallback>* StringCallbacks = StringSubscriptions.Find(InName))
-	{
-		for (TArray<FStringParamCallback>::TIterator StringCallbackIt = StringCallbacks->CreateIterator(); StringCallbacks; ++StringCallbacks)
-		{
-			if (StringCallbackIt->IsBound())
-			{
-				StringCallbackIt->Execute(InValue);
-			}
-			else
-			{
-				StringCallbackIt.RemoveCurrent();
-			}
-		}
-	}
 }
 
 void UAudioComponentGroup::SetStringArrayParameter(FName InName, const TArray<FString>& InValue)
@@ -452,6 +457,31 @@ void UAudioComponentGroup::SetParameter(FAudioParameter&& InValue)
 void UAudioComponentGroup::SetParameters(TArray<FAudioParameter>&& InValues)
 {
 	InValues.RemoveAll([](FAudioParameter& Param) { return Param.ParamName.IsNone(); });
+
+	for (FAudioParameter& Parameter : InValues)
+	{
+		switch (Parameter.ParamType)
+		{
+		case EAudioParameterType::Boolean:
+			if (GetBoolParamValue(Parameter.ParamName) != Parameter.BoolParam)
+			{
+				ExecuteBoolParamSubscriptions(Parameter);
+			}
+			break;
+		case EAudioParameterType::Trigger:
+			ExecuteEventSubscriptions(Parameter.ParamName);
+			break;
+		case EAudioParameterType::String:
+			if (GetStringParamValue(Parameter.ParamName) != Parameter.StringParam)
+			{
+				ExecuteStringParamSubscriptions(Parameter);
+			}
+			break;
+		default:
+			break;
+		}
+	}
+	
 	FAudioParameter::Merge(MoveTemp(InValues), ParamsToSet);
 }
 
@@ -538,6 +568,75 @@ float UAudioComponentGroup::GetComponentVolume() const
 	return CachedModifier.Volume;
 }
 
+void UAudioComponentGroup::ExecuteStringParamSubscriptions(const FAudioParameter& StringParam)
+{
+	if (StringParam.ParamType != EAudioParameterType::String)
+	{
+		return;
+	}
+	
+	if (TArray<FStringParamCallback>* StringCallbacks = StringSubscriptions.Find(StringParam.ParamName))
+	{
+		for (TArray<FStringParamCallback>::TIterator StringCallbackIt = StringCallbacks->CreateIterator(); StringCallbackIt; ++StringCallbackIt)
+		{
+			if (StringCallbackIt->IsBound())
+			{
+				StringCallbackIt->Execute(StringParam.StringParam);
+			}
+			else
+			{
+				StringCallbackIt.RemoveCurrent();
+			}
+		}
+	}
+}
+
+void UAudioComponentGroup::ExecuteBoolParamSubscriptions(const FAudioParameter& BoolParam)
+{
+	if (BoolParam.ParamType != EAudioParameterType::Boolean)
+	{
+		return;
+	}
+	
+	if (TArray<FBoolParamCallback>* BoolCallbacks = BoolSubscriptions.Find(BoolParam.ParamName))
+	{
+		for (TArray<FBoolParamCallback>::TIterator BoolCallbackIt = BoolCallbacks->CreateIterator(); BoolCallbackIt; ++BoolCallbackIt)
+		{
+			if (BoolCallbackIt->IsBound())
+			{
+				BoolCallbackIt->Execute(BoolParam.BoolParam);
+			}
+			else
+			{
+				BoolCallbackIt.RemoveCurrent();
+			}
+		}
+	}
+}
+
+void UAudioComponentGroup::ExecuteEventSubscriptions(const FName EventName)
+{
+	if (EventName.IsNone())
+	{
+		return;
+	}
+	
+	if (TArray<FSoundCallback>* Callbacks = EventSubscriptions.Find(EventName))
+	{
+		for (TArray<FSoundCallback>::TIterator CallbackIt = Callbacks->CreateIterator(); CallbackIt; ++CallbackIt)
+		{
+			if (CallbackIt->IsBound())
+			{
+				CallbackIt->Execute(EventName);
+			}
+			else
+			{
+				CallbackIt.RemoveCurrent();
+			}
+		}
+	}
+}
+
 const FAudioParameter* UAudioComponentGroup::GetParamInternal(const FName ParamName) const
 {
 	return PersistentParams.FindByPredicate([ParamName](const FAudioParameter& Param) { return Param.ParamName == ParamName; });
diff --git a/Engine/Plugins/AudioGameplay/Source/AudioGameplay/Public/AudioComponentGroup.h b/Engine/Plugins/AudioGameplay/Source/AudioGameplay/Public/AudioComponentGroup.h
--- a/Engine/Plugins/AudioGameplay/Source/AudioGameplay/Public/AudioComponentGroup.h
+++ b/Engine/Plugins/AudioGameplay/Source/AudioGameplay/Public/AudioComponentGroup.h
@@ -103,10 +103,19 @@ public:
 	UFUNCTION(BlueprintPure, BlueprintCosmetic, Category = AudioComponentGroup)
 	FString GetStringParamValue(const FName ParamName) const;
 
+	UFUNCTION(BlueprintCallable, BlueprintCosmetic, Category = AudioComponentGroup)
 	void SubscribeToStringParam(const FName ParamName, FStringParamCallback Delegate);
+
+	UFUNCTION(BlueprintCallable, BlueprintCosmetic, Category = AudioComponentGroup)
 	void SubscribeToEvent(const FName EventName, FSoundCallback Delegate);
+	
+	UFUNCTION(BlueprintCallable, BlueprintCosmetic, Category = AudioComponentGroup)
 	void SubscribeToBool(const FName ParamName, FBoolParamCallback Delegate);
 
+	// remove any string, event, and bool subscriptions that are bound to this object
+	UFUNCTION(BlueprintCallable, BlueprintCosmetic, Category = AudioComponentGroup)
+	void UnsubscribeObject(const UObject* Object);
+	
 	//~ Begin IAudioParameterControllerInterface interface 
 	virtual void ResetParameters() override;
 
@@ -150,6 +159,10 @@ protected:
 
 	float GetComponentVolume() const;
 
+	void ExecuteStringParamSubscriptions(const FAudioParameter& StringParam);
+	void ExecuteBoolParamSubscriptions(const FAudioParameter& BoolParam);
+	void ExecuteEventSubscriptions(const FName EventName);
+
 	const FAudioParameter* GetParamInternal(const FName ParamName) const;
 
 	UPROPERTY(VisibleAnywhere, Category = AudioComponentGroup)
