Set AutoActivate false when creating new components in AudioComponentGroup so it doesn't try to play an empty sound.

diff --git a/Engine/Plugins/AudioGameplay/Source/AudioGameplay/Private/AudioComponentGroup.cpp b/Engine/Plugins/AudioGameplay/Source/AudioGameplay/Private/AudioComponentGroup.cpp
--- a/Engine/Plugins/AudioGameplay/Source/AudioGameplay/Private/AudioComponentGroup.cpp
+++ b/Engine/Plugins/AudioGameplay/Source/AudioGameplay/Private/AudioComponentGroup.cpp
@@ -139,6 +139,7 @@ UAudioComponent* UAudioComponentGroup::AddComponent()
 	if (AActor* Owner = GetOwner())
 	{
 		NewComponent = NewObject<UAudioComponent>(Owner);
+		NewComponent->bAutoActivate = false;
 
 		NewComponent->AttachToComponent(Owner->GetRootComponent(), FAttachmentTransformRules::SnapToTargetNotIncludingScale);
 		NewComponent->RegisterComponent();
