Add editor-only OnParameterChanged delegate to AudioParameterComponent for usage in Editor Util widgets

diff --git a/Engine/Plugins/AudioGameplay/Source/AudioGameplay/Private/AudioParameterComponent.cpp b/Engine/Plugins/AudioGameplay/Source/AudioGameplay/Private/AudioParameterComponent.cpp
--- a/Engine/Plugins/AudioGameplay/Source/AudioGameplay/Private/AudioParameterComponent.cpp
+++ b/Engine/Plugins/AudioGameplay/Source/AudioGameplay/Private/AudioParameterComponent.cpp
@@ -129,6 +129,14 @@ void UAudioParameterComponent::SetParameterInternal(FAudioParameter&& InParam)
 	// Optional logging
 	LogParameter(InParam);
 
+	// Optional broadcast (for editor-usage only)
+#if WITH_EDITORONLY_DATA
+	if (OnParameterChanged.IsBound())
+	{
+		OnParameterChanged.Broadcast(InParam);
+	}
+#endif // WITH_EDITORONLY_DATA
+
 	// Forward to any AudioComponents currently playing on this actor (if any)
 	TArray<UAudioComponent*> Components;
 	GetAllAudioComponents(Components);
diff --git a/Engine/Plugins/AudioGameplay/Source/AudioGameplay/Public/AudioParameterComponent.h b/Engine/Plugins/AudioGameplay/Source/AudioGameplay/Public/AudioParameterComponent.h
--- a/Engine/Plugins/AudioGameplay/Source/AudioGameplay/Public/AudioParameterComponent.h
+++ b/Engine/Plugins/AudioGameplay/Source/AudioGameplay/Public/AudioParameterComponent.h
@@ -52,6 +52,13 @@ public:
 	virtual void SetParameters(TArray<FAudioParameter>&& InValues) override;
 	/** End IAudioParameterControllerInterface */
 
+#if WITH_EDITORONLY_DATA
+	DECLARE_DYNAMIC_MULTICAST_DELEGATE_OneParam(FOnParameterChanged, const FAudioParameter&, Parameter);
+
+	UPROPERTY(BlueprintAssignable)
+	FOnParameterChanged OnParameterChanged;
+#endif // WITH_EDITORONLY_DATA
+
 private:
 
 	void SetParameterInternal(FAudioParameter&& InParam);
