[AudioComponentGroup] only broadcast OnComponentAdded when a component is actually created

diff --git a/Engine/Plugins/AudioGameplay/Source/AudioGameplay/Private/AudioComponentGroup.cpp b/Engine/Plugins/AudioGameplay/Source/AudioGameplay/Private/AudioComponentGroup.cpp
--- a/Engine/Plugins/AudioGameplay/Source/AudioGameplay/Private/AudioComponentGroup.cpp
+++ b/Engine/Plugins/AudioGameplay/Source/AudioGameplay/Private/AudioComponentGroup.cpp
@@ -123,11 +123,11 @@ UAudioComponent* UAudioComponentGroup::GetNextAvailableComponent()
 	if (ReturnComponent == nullptr)
 	{
 		ReturnComponent = AddComponent();
-	}
 
-	for (TScriptInterface<IAudioComponentGroupExtension> Extension : Extensions)
-	{
-		Extension->OnComponentAdded(ReturnComponent);
+		for (TScriptInterface<IAudioComponentGroupExtension> Extension : Extensions)
+		{
+			Extension->OnComponentAdded(ReturnComponent);
+		}
 	}
 
 	return ReturnComponent;
