AudioComponentGroup optimization
- don't update component parameters while virtualized
- when unvirtualized, set up to update all parameters to catch any that were missed

#rnx

diff --git a/Engine/Plugins/AudioGameplay/Source/AudioGameplay/Private/AudioComponentGroup.cpp b/Engine/Plugins/AudioGameplay/Source/AudioGameplay/Private/AudioComponentGroup.cpp
--- a/Engine/Plugins/AudioGameplay/Source/AudioGameplay/Private/AudioComponentGroup.cpp
+++ b/Engine/Plugins/AudioGameplay/Source/AudioGameplay/Private/AudioComponentGroup.cpp
@@ -214,6 +214,12 @@ void UAudioComponentGroup::DisableVirtualization()
 {
 	if (bIsVirtualized)
 	{
+		// update ParamsToSet so that all parameters are updated next go-around, but merge any pending values first
+		TArray<FAudioParameter> Values = ParamsToSet;
+		FAudioParameter::Merge(MoveTemp(Values), PersistentParams);
+		
+		ParamsToSet = PersistentParams;
+		
 		bIsVirtualized = false;
 		OnUnvirtualized.Broadcast();
 	}
@@ -547,12 +553,15 @@ void UAudioComponentGroup::UpdateComponentParameters()
 
 	TRACE_CPUPROFILER_EVENT_SCOPE(UAudioComponentGroup::UpdateComponentParameters);
 
-	IterateComponents([this](UAudioComponent* Component)
+	if (bIsVirtualized == false)
 	{
-		// todo: this could be a LOT of tiny allocs... can anything be done about that?
-		TArray<FAudioParameter> Values = ParamsToSet;
-		Component->SetParameters(MoveTemp(Values));
-	});
+		IterateComponents([this](UAudioComponent* Component)
+    	{
+    		// todo: this could be a LOT of tiny allocs... can anything be done about that?
+    		TArray<FAudioParameter> Values = ParamsToSet;
+    		Component->SetParameters(MoveTemp(Values));
+    	});
+	}
 	
 	TArray<FAudioParameter> Values = ParamsToSet;
 	FAudioParameter::Merge(MoveTemp(Values), PersistentParams);
