Adding / editing some AudioGameplay and SoundLibrary logs to make them more useful.

#rnx

diff --git a/Engine/Plugins/AudioGameplay/Source/AudioGameplay/Private/AudioGameplayComponent.cpp b/Engine/Plugins/AudioGameplay/Source/AudioGameplay/Private/AudioGameplayComponent.cpp
--- a/Engine/Plugins/AudioGameplay/Source/AudioGameplay/Private/AudioGameplayComponent.cpp
+++ b/Engine/Plugins/AudioGameplay/Source/AudioGameplay/Private/AudioGameplayComponent.cpp
@@ -49,11 +49,11 @@ bool UAudioGameplayComponent::HasPayloadType(PayloadFlags InType) const
 
 void UAudioGameplayComponent::Enable()
 {
-	UE_LOG(AudioGameplayLog, Verbose, TEXT("AudioGameplayComponent Enabled (%s)."), *GetFullName());
+	UE_LOG(AudioGameplayLog, VeryVerbose, TEXT("AudioGameplayComponent Enabled (%s)."), *GetFullName());
 }
 
 void UAudioGameplayComponent::Disable()
 {
-	UE_LOG(AudioGameplayLog, Verbose, TEXT("AudioGameplayComponent Disabled (%s)."), *GetFullName());
+	UE_LOG(AudioGameplayLog, VeryVerbose, TEXT("AudioGameplayComponent Disabled (%s)."), *GetFullName());
 }
 
diff --git a/Engine/Plugins/AudioGameplay/Source/AudioGameplay/Private/AudioGameplayModule.cpp b/Engine/Plugins/AudioGameplay/Source/AudioGameplay/Private/AudioGameplayModule.cpp
--- a/Engine/Plugins/AudioGameplay/Source/AudioGameplay/Private/AudioGameplayModule.cpp
+++ b/Engine/Plugins/AudioGameplay/Source/AudioGameplay/Private/AudioGameplayModule.cpp
@@ -10,6 +10,7 @@ DEFINE_LOG_CATEGORY(AudioGameplayLog);
 
 void FAudioGameplayModule::StartupModule()
 {
+	UE_LOG(AudioGameplayLog, Log, TEXT("%s"), ANSI_TO_TCHAR(__FUNCTION__));
 }
 
 void FAudioGameplayModule::ShutdownModule()
