Sound Handle API MVP

diff --git a/Engine/Plugins/AudioGameplay/Source/AudioGameplay/AudioGameplay.Build.cs b/Engine/Plugins/AudioGameplay/Source/AudioGameplay/AudioGameplay.Build.cs
--- a/Engine/Plugins/AudioGameplay/Source/AudioGameplay/AudioGameplay.Build.cs
+++ b/Engine/Plugins/AudioGameplay/Source/AudioGameplay/AudioGameplay.Build.cs
@@ -11,10 +11,9 @@ public class AudioGameplay : ModuleRules
 		PublicDependencyModuleNames.AddRange(
 			new string[]
 			{
-				"Core",
 				"AudioExtensions",
+				"Core",
 				"GameplayTags",
-				"AudioExtensions",
 			}
 			);
 			
diff --git a/Engine/Plugins/AudioGameplay/Source/AudioGameplay/Private/SoundHandleSubsystem.cpp b/Engine/Plugins/AudioGameplay/Source/AudioGameplay/Private/SoundHandleSubsystem.cpp
new file mode 100644
--- /dev/null
+++ b/Engine/Plugins/AudioGameplay/Source/AudioGameplay/Private/SoundHandleSubsystem.cpp
@@ -0,0 +1,105 @@
+﻿// Copyright Epic Games, Inc. All Rights Reserved.
+
+#include "SoundHandleSubsystem.h"
+
+#include "ActiveSound.h"
+#include "AudioDevice.h"
+#include "Audio/ISoundHandleOwner.h"
+
+bool USoundHandleSubsystem::ShouldCreateSubsystem(UObject* Outer) const
+{
+	return Super::ShouldCreateSubsystem(Outer) && !IsRunningDedicatedServer();
+}
+
+void USoundHandleSubsystem::Initialize(FSubsystemCollectionBase& Collection)
+{
+	Super::Initialize(Collection);
+	IModularFeatures::Get().RegisterModularFeature(GetModularFeatureName(), this);
+}
+
+void USoundHandleSubsystem::Deinitialize()
+{
+	IModularFeatures::Get().UnregisterModularFeature(GetModularFeatureName(), this);
+	ActiveHandles.Reset();
+	Super::Deinitialize();
+}
+
+void USoundHandleSubsystem::OnNotifyPendingDelete(const FActiveSound& ActiveSound)
+{
+	check(IsInAudioThread());
+	if (ActiveHandles.Num() > 0)
+	{
+		const Audio::FSoundHandleID HandleID = ActiveSound.GetInstanceID();
+		ActiveHandles.Remove(HandleID);
+		if(ISoundHandleOwner* Owner = Owners.FindAndRemoveChecked(HandleID))
+		{			
+			Owner->OnSoundHandleRemoved();
+		}
+	}
+}
+
+Audio::FSoundHandleID USoundHandleSubsystem::CreateSoundHandle(USoundBase* Sound, ISoundHandleOwner* Owner)
+{
+	if(Sound && Owner)
+	{
+		FSoundHandle NewHandle;
+		NewHandle.ActiveSound.bLocationDefined = false;
+		Audio::FSoundHandleID HandleID = NewHandle.ActiveSound.GetInstanceID();
+		NewHandle.ActiveSound.SetSound(Sound);
+		ActiveHandles.Emplace(HandleID, NewHandle);
+		Owners.Emplace(HandleID, Owner);
+		return HandleID;
+	}
+	return INDEX_NONE;
+}
+
+void USoundHandleSubsystem::SetTransform(const Audio::FSoundHandleID ID, const FTransform& Transform)
+{
+	// TODO: We should add this as an incoming change to handle on the audio thread in the update
+	if(FSoundHandle* Handle = ActiveHandles.Find(ID))
+	{
+		Handle->ActiveSound.bLocationDefined = true;
+		Handle->ActiveSound.Transform = Transform;
+	}	
+}
+
+Audio::EResult USoundHandleSubsystem::Play(const Audio::FSoundHandleID ID)
+{
+	if(FSoundHandle* Handle = ActiveHandles.Find(ID))
+	{
+		FAudioDevice* AudioDevice = GetAudioDeviceHandle().GetAudioDevice();
+		if (Handle->ActiveSound.bLocationDefined)
+		{
+			float MaxDistance = 0.0f;
+			float FocusFactor = 1.0f;
+			const FVector Location = Handle->ActiveSound.Transform.GetLocation();
+
+			const TObjectPtr<USoundBase> Sound = Handle->ActiveSound.GetSound();
+			const FSoundAttenuationSettings* AttenuationSettings = Sound->GetAttenuationSettingsToApply();
+			AudioDevice->GetMaxDistanceAndFocusFactor(Sound, GetWorld(), Location, AttenuationSettings, MaxDistance, FocusFactor);
+			
+			Handle->ActiveSound.bHasAttenuationSettings = (AttenuationSettings != nullptr);
+			if (Handle->ActiveSound.bHasAttenuationSettings)
+			{
+				Handle->ActiveSound.AttenuationSettings = *AttenuationSettings;
+				Handle->ActiveSound.FocusData.PriorityScale = AttenuationSettings->GetFocusPriorityScale(GetAudioDeviceHandle().GetAudioDevice()->GetGlobalFocusSettings(), FocusFactor);
+			}
+
+			Handle->ActiveSound.MaxDistance = MaxDistance;
+		}
+		
+		AudioDevice->AddNewActiveSound(Handle->ActiveSound);		
+		return Audio::EResult::Success;
+	}
+	
+	return Audio::EResult::Failure;
+}
+
+void USoundHandleSubsystem::Stop(const Audio::FSoundHandleID ID)
+{	
+	if(FSoundHandle* Handle = ActiveHandles.Find(ID))
+	{
+		FAudioDevice* AudioDevice = GetAudioDeviceHandle().GetAudioDevice();
+		AudioDevice->StopActiveSound(&Handle->ActiveSound);
+	}
+}
diff --git a/Engine/Plugins/AudioGameplay/Source/AudioGameplay/Public/SoundHandleSubsystem.h b/Engine/Plugins/AudioGameplay/Source/AudioGameplay/Public/SoundHandleSubsystem.h
new file mode 100644
--- /dev/null
+++ b/Engine/Plugins/AudioGameplay/Source/AudioGameplay/Public/SoundHandleSubsystem.h
@@ -0,0 +1,52 @@
+﻿// Copyright Epic Games, Inc. All Rights Reserved.
+
+#pragma once
+
+#include "ActiveSound.h"
+#include "ActiveSoundUpdateInterface.h"
+#include "Audio/ISoundHandleSystem.h"
+#include "Subsystems/AudioEngineSubsystem.h"
+
+#include "SoundHandleSubsystem.generated.h"
+
+class ISoundHandleOwner;
+
+/**
+ * An implementation of ISoundHandleSystem using AudioEngineSubsystem
+ */
+UCLASS()
+class USoundHandleSubsystem : public UAudioEngineSubsystem
+	, public IActiveSoundUpdateInterface
+	, public ISoundHandleSystem
+{
+	GENERATED_BODY()
+
+public:
+	virtual ~USoundHandleSubsystem() override = default;
+
+	//~ Begin USubsystem interface
+	virtual bool ShouldCreateSubsystem(UObject* Outer) const override;
+	virtual void Initialize(FSubsystemCollectionBase& Collection) override;
+	virtual void Deinitialize() override;
+	//~ End USubsystem interface
+	
+	//~ Begin IActiveSoundUpdateInterface
+	virtual void OnNotifyPendingDelete(const FActiveSound& ActiveSound) override;
+	//~ End IActiveSoundUpdateInterface
+
+	//~ Begin ISoundHandleSystem
+	virtual Audio::FSoundHandleID CreateSoundHandle(USoundBase* Sound, ISoundHandleOwner* Owner) override;
+	virtual void SetTransform(Audio::FSoundHandleID ID, const FTransform& Transform) override;
+	virtual Audio::EResult Play(Audio::FSoundHandleID ID) override;
+	virtual void Stop(Audio::FSoundHandleID ID) override;
+	//~ End ISoundHandleSystem
+
+protected:
+	struct FSoundHandle
+	{
+		FActiveSound ActiveSound;
+	};
+	
+	TMap<Audio::FSoundHandleID, FSoundHandle> ActiveHandles;
+	TMap<Audio::FSoundHandleID, ISoundHandleOwner*> Owners;
+};
diff --git a/Engine/Source/Runtime/Engine/Public/Audio/ISoundHandleOwner.h b/Engine/Source/Runtime/Engine/Public/Audio/ISoundHandleOwner.h
new file mode 100644
--- /dev/null
+++ b/Engine/Source/Runtime/Engine/Public/Audio/ISoundHandleOwner.h
@@ -0,0 +1,41 @@
+﻿// Copyright Epic Games, Inc. All Rights Reserved.
+
+#pragma once
+
+#include "Audio.h"
+#include "Audio/ISoundHandleSystem.h"
+
+/**
+ * An interface to be used for any classes that we want to respond to SoundHandle updates
+ * TODO: Add a child of this that uses World?
+ */
+class ISoundHandleOwner
+{
+public:
+	virtual ~ISoundHandleOwner() = default;
+	
+	Audio::EResult CreateSoundHandle(USoundBase* Sound, Audio::FSoundHandleID& OutSoundHandleID)
+	{
+		if(Sound)
+		{
+			for(ISoundHandleSystem* Interface : ISoundHandleSystem::GetRegisteredInterfaces())
+			{
+				OutSoundHandleID = Interface->CreateSoundHandle(Sound, this);
+				return Audio::EResult::Success;
+			}
+		}
+		else
+		{
+			UE_LOG(LogAudio, Warning, TEXT("Couldn't create Sound Handle because we didn't pass through a valid sound!"));
+		}
+		
+		UE_LOG(LogAudio, Warning, TEXT("Couldn't create Sound Handle because we couldn't find any ISoundHandleSystems!"));
+		return Audio::EResult::Failure;
+	}
+	
+	/**
+	 * Called when the sound handle is being removed from a ISoundHandle implementer,
+	 * i.e. when the active sound it handles is notified for delete
+	 */
+	virtual void OnSoundHandleRemoved() {}
+};
diff --git a/Engine/Source/Runtime/Engine/Public/Audio/ISoundHandleSystem.h b/Engine/Source/Runtime/Engine/Public/Audio/ISoundHandleSystem.h
new file mode 100644
--- /dev/null
+++ b/Engine/Source/Runtime/Engine/Public/Audio/ISoundHandleSystem.h
@@ -0,0 +1,81 @@
+﻿// Copyright Epic Games, Inc. All Rights Reserved.
+
+#pragma once
+
+#include "Features/IModularFeatures.h"
+
+class ISoundHandleOwner;
+class USoundBase;
+
+namespace Audio
+{
+	using FSoundHandleID = uint32;
+
+	enum class EResult : uint8
+	{
+		Success = 0,
+		Failure = 1,
+		NotImplemented = 2
+	};
+}
+
+/**
+ * This interface should be used with systems aiming to create Sound Handles: gameplay thread representations of sounds
+ */
+class ISoundHandleSystem : public IModularFeature
+{
+public:
+	/**	
+	 * Get the name of all Sound Handle implementations in the Modular Features registry.
+	 * @return "SoundHandles"
+	 */
+	static FName GetModularFeatureName()
+	{
+		static FName FeatureName = FName(TEXT("SoundHandles"));
+		return FeatureName;
+	}
+
+	virtual ~ISoundHandleSystem() = default;
+
+	/**
+	 * @brief This should create a sound and we should map it to an identifier.
+	 * Then return the identifier to the owner to keep track of
+	 * 
+	 * @param Sound
+	 * @param Owner 
+	 * @return 
+	 */
+	virtual Audio::FSoundHandleID CreateSoundHandle(USoundBase* Sound, ISoundHandleOwner* Owner) { return INDEX_NONE; }
+
+	/**
+	 * @brief Set the transform on the sound that is represented by the sound handle with the following ID
+	 * 
+	 * @param ID 
+	 * @param Transform 
+	 */
+	virtual void SetTransform(Audio::FSoundHandleID ID, const FTransform& Transform) {}
+
+	/**
+	 * @brief Play the sound that is represented by the sound handle with the following ID
+	 *
+	 * @param ID 
+	 */
+	virtual Audio::EResult Play(Audio::FSoundHandleID ID) { return Audio::EResult::NotImplemented; }
+
+	/**
+	 * @brief Stop the sound that is represented by the sound handle with the following ID
+	 *
+	 * @param ID 
+	 */
+	virtual void Stop(Audio::FSoundHandleID ID) {}
+
+	static TArray<ISoundHandleSystem*> GetRegisteredInterfaces()
+	{
+		// Get all IAudioHandleInterface implementations
+		IModularFeatures::Get().LockModularFeatureList();
+		TArray<ISoundHandleSystem*> RegisteredInterfaces = IModularFeatures::Get().GetModularFeatureImplementations<ISoundHandleSystem>(GetModularFeatureName());
+		IModularFeatures::Get().UnlockModularFeatureList();
+		return RegisteredInterfaces;
+	}
+};
+
