use AutoManageAttachment for pooled audio components in AudioComponentGroup

#rnx

diff --git a/Engine/Plugins/AudioGameplay/Source/AudioGameplay/Private/AudioComponentGroup.cpp b/Engine/Plugins/AudioGameplay/Source/AudioGameplay/Private/AudioComponentGroup.cpp
--- a/Engine/Plugins/AudioGameplay/Source/AudioGameplay/Private/AudioComponentGroup.cpp
+++ b/Engine/Plugins/AudioGameplay/Source/AudioGameplay/Private/AudioComponentGroup.cpp
@@ -140,8 +140,19 @@ UAudioComponent* UAudioComponentGroup::AddComponent()
 	{
 		NewComponent = NewObject<UAudioComponent>(Owner);
 		NewComponent->bAutoActivate = false;
+		NewComponent->bAutoManageAttachment = true;
 
-		NewComponent->AttachToComponent(Owner->GetRootComponent(), FAttachmentTransformRules::SnapToTargetNotIncludingScale);
+		USceneComponent* RootComponent = Owner->GetRootComponent();
+		NewComponent->AttachToComponent(RootComponent, FAttachmentTransformRules::SnapToTargetNotIncludingScale);
+
+		if (RootComponent)
+		{
+			NewComponent->AutoAttachParent = RootComponent->GetAttachParent();
+		}
+		
+		NewComponent->AutoAttachLocationRule = EAttachmentRule::KeepRelative;
+		NewComponent->AutoAttachRotationRule = EAttachmentRule::KeepRelative;
+		
 		NewComponent->RegisterComponent();
 
 		ApplyParams(NewComponent);
