AudioComponentGroup - don't spawn a component group in the actor hierarchy any higher than a pawn. prevents attaching to a playercontroller in cases where one isn't pre-attached.

#rnx

diff --git a/Engine/Plugins/AudioGameplay/Source/AudioGameplay/Private/AudioComponentGroup.cpp b/Engine/Plugins/AudioGameplay/Source/AudioGameplay/Private/AudioComponentGroup.cpp
--- a/Engine/Plugins/AudioGameplay/Source/AudioGameplay/Private/AudioComponentGroup.cpp
+++ b/Engine/Plugins/AudioGameplay/Source/AudioGameplay/Private/AudioComponentGroup.cpp
@@ -5,6 +5,7 @@
 
 #include "Components/AudioComponent.h"
 #include "GameFramework/Actor.h"
+#include "GameFramework/Pawn.h"
 
 #include UE_INLINE_GENERATED_CPP_BY_NAME(AudioComponentGroup)
 
@@ -34,6 +35,13 @@ UAudioComponentGroup* UAudioComponentGroup::StaticGetOrCreateComponentGroup(AAct
 		}
 
 		HighestValidOwner = Owner;
+
+		// prevents sound groups on pawns from spawning on owning playercontrollers
+		if (HighestValidOwner->IsA(APawn::StaticClass()))
+		{
+			break;
+		}
+		
 		Owner = Owner->GetOwner();
 	}
 
