add method to remove previously added external component. used to fix an issue where gameplay cues with audio get re-used instead of cleaning up, causing its audio component to be registered to multiple component groups at once

diff --git a/Engine/Plugins/AudioGameplay/Source/AudioGameplay/Private/AudioComponentGroup.cpp b/Engine/Plugins/AudioGameplay/Source/AudioGameplay/Private/AudioComponentGroup.cpp
--- a/Engine/Plugins/AudioGameplay/Source/AudioGameplay/Private/AudioComponentGroup.cpp
+++ b/Engine/Plugins/AudioGameplay/Source/AudioGameplay/Private/AudioComponentGroup.cpp
@@ -220,6 +220,16 @@ void UAudioComponentGroup::AddExternalComponent(UAudioComponent* ComponentToAdd)
 	}
 }
 
+void UAudioComponentGroup::RemoveExternalComponent(UAudioComponent* ComponentToRemove)
+{
+	if (ComponentToRemove == nullptr)
+	{
+		return;
+	}
+
+	ExternalComponents.Remove(ComponentToRemove);
+}
+
 void UAudioComponentGroup::EnableVirtualization()
 {
 	if (!bIsVirtualized)
diff --git a/Engine/Plugins/AudioGameplay/Source/AudioGameplay/Public/AudioComponentGroup.h b/Engine/Plugins/AudioGameplay/Source/AudioGameplay/Public/AudioComponentGroup.h
--- a/Engine/Plugins/AudioGameplay/Source/AudioGameplay/Public/AudioComponentGroup.h
+++ b/Engine/Plugins/AudioGameplay/Source/AudioGameplay/Public/AudioComponentGroup.h
@@ -67,6 +67,9 @@ public:
 	UFUNCTION(BlueprintCallable, BlueprintCosmetic, Category = AudioComponentGroup)
 	void AddExternalComponent(UAudioComponent* ComponentToAdd);
 
+	UFUNCTION(BlueprintCallable, BlueprintCosmetic, Category = AudioComponentGroup)
+	void RemoveExternalComponent(UAudioComponent* ComponentToRemove);
+
 	UFUNCTION(BlueprintCallable, BlueprintCosmetic, Category = AudioComponentGroup)
 	void EnableVirtualization();
 
