AudioComponentGroup - don't override cached modifier values with virtualization to keep everything in sync with the cached value so it can update naturally on tick. Always update external components whether active or not, since they do not benefit from resetting when requested like pooled components do.

diff --git a/Engine/Plugins/AudioGameplay/Source/AudioGameplay/Private/AudioComponentGroup.cpp b/Engine/Plugins/AudioGameplay/Source/AudioGameplay/Private/AudioComponentGroup.cpp
--- a/Engine/Plugins/AudioGameplay/Source/AudioGameplay/Private/AudioComponentGroup.cpp
+++ b/Engine/Plugins/AudioGameplay/Source/AudioGameplay/Private/AudioComponentGroup.cpp
@@ -165,9 +165,7 @@ UAudioComponent* UAudioComponentGroup::AddComponent()
 
 		ApplyParams(NewComponent);
 
-		FAudioComponentModifier Modifier = CachedModifier;
-		Modifier.Volume = GetComponentVolume();
-		ApplyModifiers(NewComponent, Modifier);
+		ApplyModifiers(NewComponent, CachedModifier);
 
 		Components.Add(NewComponent);
 	}
@@ -185,9 +183,7 @@ UAudioComponent* UAudioComponentGroup::ResetComponent(UAudioComponent* Component
 
 		ApplyParams(Component);
 
-		FAudioComponentModifier Modifier = CachedModifier;
-		Modifier.Volume = GetComponentVolume();
-		ApplyModifiers(Component, Modifier);
+		ApplyModifiers(Component, CachedModifier);
 	}
 
 	return Component;
@@ -212,9 +208,7 @@ void UAudioComponentGroup::AddExternalComponent(UAudioComponent* ComponentToAdd)
 	{
 		ApplyParams(ComponentToAdd);
 
-		FAudioComponentModifier Modifier = CachedModifier;
-		Modifier.Volume = GetComponentVolume();
-		ApplyModifiers(ComponentToAdd, Modifier);
+		ApplyModifiers(ComponentToAdd, CachedModifier);
 
 		ExternalComponents.Add(TWeakObjectPtr<UAudioComponent>(ComponentToAdd));
 	}
@@ -566,10 +560,7 @@ void UAudioComponentGroup::IterateComponents(const TFunction<void(UAudioComponen
 	ExternalComponents.RemoveAll([](const TWeakObjectPtr<UAudioComponent> WeakComponent) { return !WeakComponent.IsValid(); });
 	for (TWeakObjectPtr<UAudioComponent>& WeakComponent : ExternalComponents)
 	{
-		if (WeakComponent->IsActive())
-		{
-			OnIterate(WeakComponent.Get());
-		}
+		OnIterate(WeakComponent.Get());
 	}
 }
 
