[Perf] - Optimization for unneccesary set params
+ Fix issue where parameters could potentially be dropped / reset when sent to metasounds
+ CVar for disable - au.AudioParameterComponent.SetParamOnlyOnValueChange 0

#localization none
#tests none

diff --git a/Engine/Plugins/AudioGameplay/Source/AudioGameplay/Private/AudioParameterComponent.cpp b/Engine/Plugins/AudioGameplay/Source/AudioGameplay/Private/AudioParameterComponent.cpp
--- a/Engine/Plugins/AudioGameplay/Source/AudioGameplay/Private/AudioParameterComponent.cpp
+++ b/Engine/Plugins/AudioGameplay/Source/AudioGameplay/Private/AudioParameterComponent.cpp
@@ -7,6 +7,109 @@
 
 #include UE_INLINE_GENERATED_CPP_BY_NAME(AudioParameterComponent)
 
+namespace AudioParameterComponentConsoleVariables
+{
+	int32 bSetParamOnlyOnValueChange = 1;
+	FAutoConsoleVariableRef CVarSetParamOnlyOnValueChange(
+		TEXT("au.AudioParameterComponent.SetParamOnlyOnValueChange"),
+		bSetParamOnlyOnValueChange,
+		TEXT("Only sets parameters when the underlying value has changed.\n0: Disable, 1: Enable (default)"),
+		ECVF_Default);
+}
+
+namespace AudioParameterComponentUtils
+{
+	// This one is a bit weird, but we combine params by merging them - and we currently 
+	// always take the parameter type of the thing being merged in.  This means for something to be 
+	// Equivalent (for a no-op type optimization) it has to have the same ParamType
+	bool TestEquivalence(const FAudioParameter& Lhs, const FAudioParameter& Rhs)
+	{
+		if (Lhs.ParamType != Rhs.ParamType)
+		{
+			return false;
+		}
+
+		switch (Rhs.ParamType)
+		{
+		case EAudioParameterType::Boolean:
+		{
+			return Lhs.BoolParam == Rhs.BoolParam;
+		}
+		break;
+
+		case EAudioParameterType::BooleanArray:
+		{
+			return Lhs.ArrayBoolParam == Rhs.ArrayBoolParam;
+		}
+		break;
+
+		case EAudioParameterType::Float:
+		{
+			return FMath::IsNearlyEqual(Lhs.FloatParam, Rhs.FloatParam);
+		}
+		break;
+
+		case EAudioParameterType::FloatArray:
+		{
+			return Lhs.ArrayFloatParam == Rhs.ArrayFloatParam;
+		}
+		break;
+
+		case EAudioParameterType::Integer: 
+		case EAudioParameterType::NoneArray:
+		{
+			return Lhs.IntParam == Rhs.IntParam;
+		}
+		break;
+
+		case EAudioParameterType::IntegerArray:
+		{
+			return Lhs.ArrayIntParam == Rhs.ArrayIntParam;
+		}
+		break;
+
+		case EAudioParameterType::None:
+		{
+			// Moved string comp to end hoping for lazy eval
+			return FMath::IsNearlyEqual(Lhs.FloatParam, Rhs.FloatParam) && Lhs.BoolParam == Rhs.BoolParam && Lhs.IntParam == Rhs.IntParam &&
+				Lhs.ObjectParam == Rhs.ObjectParam && Lhs.ArrayFloatParam == Rhs.ArrayFloatParam && Lhs.ArrayBoolParam == Rhs.ArrayBoolParam &&
+				Lhs.ArrayIntParam == Rhs.ArrayIntParam && Lhs.ArrayObjectParam == Rhs.ArrayObjectParam &&
+				Lhs.StringParam == Rhs.StringParam && Lhs.ArrayStringParam == Rhs.ArrayStringParam;
+		}
+		break;
+
+		case EAudioParameterType::Object:
+		{
+			return Lhs.ObjectParam == Rhs.ObjectParam;
+		}
+		break;
+
+		case EAudioParameterType::ObjectArray:
+		{
+			return Lhs.ArrayObjectParam == Rhs.ArrayObjectParam;
+		}
+		break;
+
+		case EAudioParameterType::String:
+		{
+			return Lhs.StringParam == Rhs.StringParam;
+		}
+		break;
+
+		case EAudioParameterType::StringArray:
+		{
+			return Lhs.ArrayStringParam == Rhs.ArrayStringParam;
+		}
+		break;
+
+		default:
+		break;
+
+		}
+
+		return false;
+	}
+}
 
 DEFINE_LOG_CATEGORY(LogAudioParameterComponent);
 
@@ -120,8 +223,23 @@ void UAudioParameterComponent::SetParameterInternal(FAudioParameter&& InParam)
 		return;
 	}
 
+	const int32 StartingParamCount = Parameters.Num();
 	if (FAudioParameter* CurrentParam = FAudioParameter::FindOrAddParam(Parameters, InParam.ParamName))
 	{
+		if (AudioParameterComponentConsoleVariables::bSetParamOnlyOnValueChange != 0)
+		{
+			if (StartingParamCount == Parameters.Num())
+			{
+				// Early out if we did not add a parameter and this one matches.
+				// We always take the type of the param (see merge below), so the utility function below
+				// is only going to test based on changes to the right hand side 
+				if  (AudioParameterComponentUtils::TestEquivalence(*CurrentParam, InParam))
+				{
+					return;
+				}
+			}
+		}
+
 		constexpr bool bInTakeName = false;
 		CurrentParam->Merge(InParam, bInTakeName);
 	}
@@ -141,6 +259,7 @@ void UAudioParameterComponent::SetParameterInternal(FAudioParameter&& InParam)
 	TArray<UAudioComponent*> Components;
 	GetAllAudioComponents(Components);
 
+	TArray<FAudioParameter> ParametersCopy;
 	for (auto& Component : Components)
 	{
 		if (Component == nullptr)
@@ -152,19 +271,19 @@ void UAudioParameterComponent::SetParameterInternal(FAudioParameter&& InParam)
 		{
 			if (FAudioDevice* AudioDevice = Component->GetAudioDevice())
 			{
+				ParametersCopy = Parameters;
 				if (USoundBase* Sound = Component->GetSound())
 				{
 					static const FName ProxyFeatureName("AudioParameterComponent");
-					Sound->InitParameters(Parameters, ProxyFeatureName);
+					Sound->InitParameters(ParametersCopy, ProxyFeatureName);
 				}
 
 				// Prior call to InitParameters can prune parameters if they are
 				// invalid, so check here to avoid unnecessary pass of empty array.
-				if (!Parameters.IsEmpty())
+				if (!ParametersCopy.IsEmpty())
 				{
 					DECLARE_CYCLE_STAT(TEXT("FAudioThreadTask.SoundParameterControllerInterface.SetParameters"), STAT_AudioSetParameters, STATGROUP_AudioThreadCommands);
 
-					TArray<FAudioParameter> ParametersCopy = Parameters;
 					AudioDevice->SendCommandToActiveSounds(Component->GetInstanceOwnerID(), [AudioDevice, Params = MoveTemp(ParametersCopy)](FActiveSound& ActiveSound)
 					{
 						if (Audio::IParameterTransmitter* Transmitter = ActiveSound.GetTransmitter())
