FilteredGameplayTagContainer

#virtualized

diff --git a/Engine/Plugins/AudioGameplay/AudioGameplay.uplugin b/Engine/Plugins/AudioGameplay/AudioGameplay.uplugin
--- a/Engine/Plugins/AudioGameplay/AudioGameplay.uplugin
+++ b/Engine/Plugins/AudioGameplay/AudioGameplay.uplugin
@@ -10,7 +10,7 @@
 	"DocsURL": "",
 	"MarketplaceURL": "",
 	"SupportURL": "",
-	"CanContainContent": false,
+	"CanContainContent": true,
 	"IsBetaVersion": true,
 	"Installed": false,
 	"Modules": [
@@ -18,6 +18,11 @@
 			"Name": "AudioGameplay",
 			"Type": "Runtime",
 			"LoadingPhase": "Default"
+		},
+		{
+			"Name": "AudioGameplayTests",
+			"Type": "Editor",
+			"LoadingPhase": "Default"
 		}
 	]
 }
\ No newline at end of file
diff --git a/Engine/Plugins/AudioGameplay/Source/AudioGameplay/Private/FilteredGameplayTagContainer.cpp b/Engine/Plugins/AudioGameplay/Source/AudioGameplay/Private/FilteredGameplayTagContainer.cpp
new file mode 100644
--- /dev/null
+++ b/Engine/Plugins/AudioGameplay/Source/AudioGameplay/Private/FilteredGameplayTagContainer.cpp
@@ -0,0 +1,89 @@
+// Copyright Epic Games, Inc. All Rights Reserved.
+
+#include "FilteredGameplayTagContainer.h"
+
+bool FFilteredGameplayTagContainer::AddTagFiltered(const FilteredGPTContainer::FGPTagQueryPair& InPair)
+{
+	if (!AddTagToContainers(InPair))
+	{
+		return false;
+	}
+
+	FilterTags();
+	return true;
+}
+
+bool FFilteredGameplayTagContainer::AddTagFiltered(const FGameplayTag& InTag, const FGameplayTagQuery& InQuery /*= FGameplayTagQuery()*/)
+{
+	return AddTagFiltered(FilteredGPTContainer::FGPTagQueryPair(InTag, InQuery));
+}
+
+
+bool FFilteredGameplayTagContainer::RemoveTagFiltered(const FGameplayTag& InTag)
+{
+	if (!RemoveTagFromContainers(InTag))
+	{
+		return false;
+	}
+
+	FilterTags();
+	return true;
+}
+
+bool FFilteredGameplayTagContainer::AddTagToContainers(const FilteredGPTContainer::FGPTagQueryPair& InPair)
+{
+	if (!TagContainer.HasTagExact(InPair.Key) && (InPair.Value.IsEmpty() || TagContainer.MatchesQuery(InPair.Value)))
+	{
+		GPTagQueryPairContainer.Add(InPair);
+		TagContainer.AddTagFast(InPair.Key);
+		OnGameplayTagAdded.Broadcast(InPair.Key);
+		return true;
+	}
+
+	return false;
+}
+
+bool FFilteredGameplayTagContainer::RemoveTagFromContainers(const FGameplayTag& Tag)
+{
+	for (int32 i = GPTagQueryPairContainer.Num() - 1; i >= 0; --i)
+	{
+		if (GPTagQueryPairContainer[i].Key == Tag)
+		{
+			GPTagQueryPairContainer.RemoveAt(i);
+			TagContainer.RemoveTag(Tag);
+			OnGameplayTagRemoved.Broadcast(Tag);
+			return true;
+
+		}
+	}
+
+	return false;	
+}
+
+void FFilteredGameplayTagContainer::FilterTags()
+{
+	for (int32 i = GPTagQueryPairContainer.Num() - 1; i >= 0; --i)
+	{
+		const FilteredGPTContainer::FGPTagQueryPair& TagQueryPair = GPTagQueryPairContainer[i];
+
+		if (TagQueryPair.Value.IsEmpty())
+		{
+			continue;
+		}
+
+		if (!TagContainer.MatchesQuery(TagQueryPair.Value))
+		{
+			FGameplayTag TagToRemove = TagQueryPair.Key;
+			GPTagQueryPairContainer.RemoveAt(i);
+			TagContainer.RemoveTag(TagToRemove);
+			OnGameplayTagRemoved.Broadcast(TagToRemove);
+		}
+	}
+
+	CheckContainersConsistency();
+}
+
+void FFilteredGameplayTagContainer::CheckContainersConsistency()
+{
+	ensureMsgf(GPTagQueryPairContainer.Num() == TagContainer.Num(), TEXT("Expected Event and Tag Container to have the same lenght. GPTagQueryPairContainer: %d, TagContainer: %d"), GPTagQueryPairContainer.Num(), TagContainer.Num());
+}
diff --git a/Engine/Plugins/AudioGameplay/Source/AudioGameplay/Public/FilteredGameplayTagContainer.h b/Engine/Plugins/AudioGameplay/Source/AudioGameplay/Public/FilteredGameplayTagContainer.h
new file mode 100644
--- /dev/null
+++ b/Engine/Plugins/AudioGameplay/Source/AudioGameplay/Public/FilteredGameplayTagContainer.h
@@ -0,0 +1,86 @@
+// Copyright Epic Games, Inc. All Rights Reserved.
+
+#include "Delegates/Delegate.h" 
+#include "GameplayTagContainer.h"
+
+DECLARE_MULTICAST_DELEGATE_OneParam(FOnFilteredGameplayTagAdded, const FGameplayTag)
+DECLARE_MULTICAST_DELEGATE_OneParam(FOnFilteredGameplayTagRemoved, const FGameplayTag)
+
+namespace FilteredGPTContainer
+{
+	typedef TPair<FGameplayTag, FGameplayTagQuery> FGPTagQueryPair;
+}
+
+/*
+*	A container of FGameplayTag associated to a FGameplayTagQuery. 
+*	Tags will only be added if the accompanying query is satisfied. 
+*	Whenever a tag is added or removed, the tags in the container are "Filtered":
+*	i.e. their query is checked to see if they should still remain in the array. 
+* 
+*	Note:
+*	Tags are unique
+*	Filtering is evaluated from last to first element in the array
+*/
+class AUDIOGAMEPLAY_API FFilteredGameplayTagContainer
+{
+public:
+
+	/**
+	* Adds a tag to the container if the query matches. 
+	* If the tag is added, proceeds to evaluate the conditions on the contained tags to check 
+	* if they should remain or not.
+	*
+	* @param InPair			The tag and query to try and add to this container
+	*
+	* @return True if tag was added
+	*/
+	bool AddTagFiltered(const FilteredGPTContainer::FGPTagQueryPair& InPair);
+	
+	/**
+	* Adds a tag to the container if the query matches.
+	* If the tag is added, proceeds to evaluate the conditions on the contained tags to check
+	* if they should be kept or not.
+	*
+	* @param InTag			The tag to try add to this container
+	* @param InQuery		The query to determine if this tag should be added to the container
+	*
+	* @return True if tag was added
+	*/
+	bool AddTagFiltered(const FGameplayTag& InTag, const FGameplayTagQuery& InQuery = FGameplayTagQuery());
+
+	/**
+	* Removes a tag from the container, then proceeds to evaluate the conditions on the contained tags to check
+	* if they should be kept or not.
+	*
+	* @return True if tag was removed
+	*/
+	bool RemoveTagFiltered(const FGameplayTag& InTag);
+
+	FORCEINLINE int32 Find(const FilteredGPTContainer::FGPTagQueryPair& EventToFind) const
+	{
+		return GPTagQueryPairContainer.Find(EventToFind);
+	}
+
+	FORCEINLINE int32 Num() const
+	{
+		return GPTagQueryPairContainer.Num();
+	}
+
+	FORCEINLINE bool IsEmpty() const
+	{
+		return GPTagQueryPairContainer.Num() == 0;
+	}
+
+	FOnFilteredGameplayTagAdded OnGameplayTagAdded;
+	FOnFilteredGameplayTagRemoved OnGameplayTagRemoved;
+
+private:
+
+	bool AddTagToContainers(const FilteredGPTContainer::FGPTagQueryPair& InPair);
+	bool RemoveTagFromContainers(const FGameplayTag& Tag);
+	void FilterTags();
+	void CheckContainersConsistency();
+
+	TArray<FilteredGPTContainer::FGPTagQueryPair> GPTagQueryPairContainer;
+	FGameplayTagContainer TagContainer;
+};
diff --git a/Engine/Plugins/AudioGameplay/Source/AudioGameplayTests/AudioGameplayTests.build.cs b/Engine/Plugins/AudioGameplay/Source/AudioGameplayTests/AudioGameplayTests.build.cs
new file mode 100644
--- /dev/null
+++ b/Engine/Plugins/AudioGameplay/Source/AudioGameplayTests/AudioGameplayTests.build.cs
@@ -0,0 +1,37 @@
+// Copyright Epic Games, Inc. All Rights Reserved.
+
+using UnrealBuildTool;
+
+public class AudioGameplayTests : ModuleRules
+{
+	public AudioGameplayTests(ReadOnlyTargetRules Target) : base(Target)
+	{
+		PCHUsage = ModuleRules.PCHUsageMode.UseExplicitOrSharedPCHs;
+
+		PublicDependencyModuleNames.AddRange(
+			new string[]
+			{
+				"AudioExtensions",
+				"AudioGameplay",
+				"Core",
+				"GameplayTags",
+			}
+			);
+
+
+		PrivateDependencyModuleNames.AddRange(
+			new string[]
+			{
+				"CoreUObject",
+				"Engine",
+			}
+			);
+
+
+		DynamicallyLoadedModuleNames.AddRange(
+			new string[]
+			{
+			}
+			);
+	}
+}
\ No newline at end of file
diff --git a/Engine/Plugins/AudioGameplay/Source/AudioGameplayTests/Private/AudioGameplayTestsModule.cpp b/Engine/Plugins/AudioGameplay/Source/AudioGameplayTests/Private/AudioGameplayTestsModule.cpp
new file mode 100644
--- /dev/null
+++ b/Engine/Plugins/AudioGameplay/Source/AudioGameplayTests/Private/AudioGameplayTestsModule.cpp
@@ -0,0 +1,14 @@
+// Copyright Epic Games, Inc. All Rights Reserved.
+
+#include "AudioGameplayTestsModule.h"
+#include "Modules/ModuleManager.h"
+
+void FAudioGameplayTestsModule::StartupModule()
+{
+}
+
+void FAudioGameplayTestsModule::ShutdownModule()
+{
+}
+
+IMPLEMENT_MODULE(FAudioGameplayTestsModule, AudioGameplayTests)
\ No newline at end of file
diff --git a/Engine/Plugins/AudioGameplay/Source/AudioGameplayTests/Private/FilteredGameplayTagContainerTests.cpp b/Engine/Plugins/AudioGameplay/Source/AudioGameplayTests/Private/FilteredGameplayTagContainerTests.cpp
new file mode 100644
--- /dev/null
+++ b/Engine/Plugins/AudioGameplay/Source/AudioGameplayTests/Private/FilteredGameplayTagContainerTests.cpp
@@ -0,0 +1,140 @@
+// Copyright Epic Games, Inc. All Rights Reserved.
+
+#include "Misc/AutomationTest.h"
+
+#if WITH_DEV_AUTOMATION_TESTS
+
+#include "FilteredGameplayTagContainer.h"
+
+
+namespace AudioGameplayTagsTestParams
+{
+	static const FName TestTagName1 = "AudioGameplay.Tests.FilteredGameplayTagContainer.TestTag1";
+	static const FName TestTagName2 = "AudioGameplay.Tests.FilteredGameplayTagContainer.TestTag2";
+	static const FName TestTagName3 = "AudioGameplay.Tests.FilteredGameplayTagContainer.TestTag3";
+}
+
+IMPLEMENT_SIMPLE_AUTOMATION_TEST(FAudioGameplayTagTests_TestAddTagFiltered, "AudioGameplay.FilteredGameplayTags.TestAddTagFiltered", EAutomationTestFlags::EditorContext | EAutomationTestFlags::EngineFilter)
+bool FAudioGameplayTagTests_TestAddTagFiltered::RunTest(const FString& Parameters)
+{
+	FGameplayTag FilteredTag1 = FGameplayTag::RequestGameplayTag(AudioGameplayTagsTestParams::TestTagName1);
+	FGameplayTag FilteredTag2 = FGameplayTag::RequestGameplayTag(AudioGameplayTagsTestParams::TestTagName3);
+
+	FFilteredGameplayTagContainer TestContainer;
+
+	bool AddTagFilteredReturn = true;
+
+	AddTagFilteredReturn &= TestContainer.AddTagFiltered(FilteredTag1);
+	bool bExpectedNumToBeOne = TestContainer.Num() == 1;
+
+	AddTagFilteredReturn &= TestContainer.AddTagFiltered(FilteredTag2);
+	bool bExpectedNumToBeTwo = TestContainer.Num() == 2;
+
+	return AddTagFilteredReturn && bExpectedNumToBeOne && bExpectedNumToBeTwo;
+}
+
+IMPLEMENT_SIMPLE_AUTOMATION_TEST(FAudioGameplayTagTests_TestRemoveTagFiltered, "AudioGameplay.FilteredGameplayTags.TestRemoveTagFiltered", EAutomationTestFlags::EditorContext | EAutomationTestFlags::EngineFilter)
+bool FAudioGameplayTagTests_TestRemoveTagFiltered::RunTest(const FString& Parameters)
+{
+	FGameplayTag FilteredTag1 = FGameplayTag::RequestGameplayTag(AudioGameplayTagsTestParams::TestTagName1);
+	FGameplayTag FilteredTag2 = FGameplayTag::RequestGameplayTag(AudioGameplayTagsTestParams::TestTagName3);
+
+	FFilteredGameplayTagContainer TestContainer;
+
+	TestContainer.AddTagFiltered(FilteredTag1);
+	TestContainer.AddTagFiltered(FilteredTag2);
+
+	bool bExpectedNumToBeTwo = TestContainer.Num() == 2;
+
+	bool bRemoveTagFilteredReturn = TestContainer.RemoveTagFiltered(FilteredTag1);
+	bool bExpectedNumToBeOne = TestContainer.Num() == 1;
+
+	return bRemoveTagFilteredReturn && bExpectedNumToBeOne && bExpectedNumToBeTwo;
+}
+
+IMPLEMENT_SIMPLE_AUTOMATION_TEST(FAudioGameplayTagTests_TestAddWithSameName, "AudioGameplay.FilteredGameplayTags.TestAddWithSameName", EAutomationTestFlags::EditorContext | EAutomationTestFlags::EngineFilter)
+bool FAudioGameplayTagTests_TestAddWithSameName::RunTest(const FString& Parameters)
+{
+	FGameplayTag FilteredTag1 = FGameplayTag::RequestGameplayTag(AudioGameplayTagsTestParams::TestTagName1);
+
+	FFilteredGameplayTagContainer TestContainer;
+
+	const bool bFirstAddPass = TestContainer.AddTagFiltered(FilteredTag1);
+	const bool bSecondAddFails = !TestContainer.AddTagFiltered(FilteredTag1);
+	bool bExpectedNumToBeOne = TestContainer.Num() == 1;
+
+	return bFirstAddPass && bSecondAddFails && bExpectedNumToBeOne;
+}
+
+IMPLEMENT_SIMPLE_AUTOMATION_TEST(FAudioGameplayTagTests_TestRemoveWithNoMatch, "AudioGameplay.FilteredGameplayTags.TestRemoveWithNoMatch", EAutomationTestFlags::EditorContext | EAutomationTestFlags::EngineFilter)
+bool FAudioGameplayTagTests_TestRemoveWithNoMatch::RunTest(const FString& Parameters)
+{
+	FGameplayTag FilteredTag1 = FGameplayTag::RequestGameplayTag(AudioGameplayTagsTestParams::TestTagName1);
+
+	FFilteredGameplayTagContainer TestContainer;
+
+	const bool bRemoveFails = !TestContainer.RemoveTagFiltered(FilteredTag1);
+	const bool bEmptyContainer = TestContainer.IsEmpty();
+
+
+	return bRemoveFails && bEmptyContainer;
+}
+
+IMPLEMENT_SIMPLE_AUTOMATION_TEST(FAudioGameplayTagTests_TestAddTagFilteredWithInvalidCondition, "AudioGameplay.FilteredGameplayTags.TestAddTagFilteredWithInvalidCondition", EAutomationTestFlags::EditorContext | EAutomationTestFlags::EngineFilter)
+bool FAudioGameplayTagTests_TestAddTagFilteredWithInvalidCondition::RunTest(const FString& Parameters)
+{
+	FGameplayTag FilteredTag1 = FGameplayTag::RequestGameplayTag(AudioGameplayTagsTestParams::TestTagName1);
+
+	FGameplayTagQuery NotFilteredTag1Query = FGameplayTagQuery::BuildQuery(FGameplayTagQueryExpression().NoTagsMatch().AddTag(FilteredTag1));
+
+	FilteredGPTContainer::FGPTagQueryPair FilteredTag2(FGameplayTag::RequestGameplayTag(AudioGameplayTagsTestParams::TestTagName2), NotFilteredTag1Query);
+
+	FFilteredGameplayTagContainer TestContainer;
+	TestContainer.AddTagFiltered(FilteredTag1);
+	bool bAddTagFilteredReturn = TestContainer.AddTagFiltered(FilteredTag2);
+
+	return !bAddTagFilteredReturn;
+}
+
+
+IMPLEMENT_SIMPLE_AUTOMATION_TEST(FAudioGameplayTagTests_TestEventRemovedWhenConditionBecomesFalseFromAdding, "AudioGameplay.FilteredGameplayTags.TestEventRemovedWhenConditionBecomesFalseFromAdding", EAutomationTestFlags::EditorContext | EAutomationTestFlags::EngineFilter)
+bool FAudioGameplayTagTests_TestEventRemovedWhenConditionBecomesFalseFromAdding::RunTest(const FString& Parameters)
+{
+	FGameplayTag FilteredTag1 = FGameplayTag::RequestGameplayTag(AudioGameplayTagsTestParams::TestTagName1);
+
+	FGameplayTagQuery NotFilteredTag1Query = FGameplayTagQuery::BuildQuery(FGameplayTagQueryExpression().NoTagsMatch().AddTag(FilteredTag1));	
+
+	FilteredGPTContainer::FGPTagQueryPair FilteredTag2(FGameplayTag::RequestGameplayTag(AudioGameplayTagsTestParams::TestTagName2), NotFilteredTag1Query);
+
+	FFilteredGameplayTagContainer TestContainer;
+	TestContainer.AddTagFiltered(FilteredTag2);
+	bool bFilteredTag2InitiallyAdded = TestContainer.Find(FilteredTag2) != INDEX_NONE; 
+	TestContainer.AddTagFiltered(FilteredTag1);
+	bool bFilteredTag2NoLongerPresent = TestContainer.Find(FilteredTag2) == INDEX_NONE;
+
+	return bFilteredTag2InitiallyAdded && bFilteredTag2NoLongerPresent;
+}
+
+IMPLEMENT_SIMPLE_AUTOMATION_TEST(FAudioGameplayTagTests_TestEventRemovedWhenConditionBecomesFalseFromRemoving, "AudioGameplay.FilteredGameplayTags.TestEventRemovedWhenConditionBecomesFalseFromRemoving", EAutomationTestFlags::EditorContext | EAutomationTestFlags::EngineFilter)
+bool FAudioGameplayTagTests_TestEventRemovedWhenConditionBecomesFalseFromRemoving::RunTest(const FString& Parameters)
+{
+	FGameplayTag FilteredTag1 = FGameplayTag::RequestGameplayTag(AudioGameplayTagsTestParams::TestTagName1);
+
+	FGameplayTagQuery FilteredTag1PresentQuery = FGameplayTagQuery::BuildQuery(FGameplayTagQueryExpression().AnyTagsMatch().AddTag(FilteredTag1));
+
+	FilteredGPTContainer::FGPTagQueryPair FilteredTag2(FGameplayTag::RequestGameplayTag(AudioGameplayTagsTestParams::TestTagName2), FilteredTag1PresentQuery);
+
+	FFilteredGameplayTagContainer TestContainer;
+	TestContainer.AddTagFiltered(FilteredTag1);
+	TestContainer.AddTagFiltered(FilteredTag2);
+	bool bFilteredTag2InitiallyAdded = TestContainer.Find(FilteredTag2) != INDEX_NONE;
+	bool bExpectedNumToBeTwo = TestContainer.Num() == 2;
+
+	TestContainer.RemoveTagFiltered(FilteredTag1);
+	bool bFilteredTag2NoLongerPresent = TestContainer.Find(FilteredTag2) == INDEX_NONE;
+	bool bExpectedNumToBeZero = TestContainer.Num() == 0;
+
+	return bFilteredTag2InitiallyAdded && bFilteredTag2NoLongerPresent && bExpectedNumToBeTwo && bExpectedNumToBeZero;
+}
+
+#endif
\ No newline at end of file
diff --git a/Engine/Plugins/AudioGameplay/Source/AudioGameplayTests/Public/AudioGameplayTestsModule.h b/Engine/Plugins/AudioGameplay/Source/AudioGameplayTests/Public/AudioGameplayTestsModule.h
new file mode 100644
--- /dev/null
+++ b/Engine/Plugins/AudioGameplay/Source/AudioGameplayTests/Public/AudioGameplayTestsModule.h
@@ -0,0 +1,18 @@
+// Copyright Epic Games, Inc. All Rights Reserved.
+
+#pragma once
+
+#include "Modules/ModuleInterface.h"
+
+/**
+ *  FAudioGameplayTestsModule
+ */
+class FAudioGameplayTestsModule : public IModuleInterface
+{
+public:
+
+	//~ Begin IModuleInterface interface
+	virtual void StartupModule() override;
+	virtual void ShutdownModule() override;
+	//~ End IModuleInterface interface
+};
\ No newline at end of file
