Add CSV tracking to AudioGameplayModule and use it in audi ovehicle controllers

diff --git a/Engine/Plugins/AudioGameplay/Source/AudioGameplay/Private/AudioGameplayModule.cpp b/Engine/Plugins/AudioGameplay/Source/AudioGameplay/Private/AudioGameplayModule.cpp
--- a/Engine/Plugins/AudioGameplay/Source/AudioGameplay/Private/AudioGameplayModule.cpp
+++ b/Engine/Plugins/AudioGameplay/Source/AudioGameplay/Private/AudioGameplayModule.cpp
@@ -3,11 +3,18 @@
 #include "AudioGameplayModule.h"
 #include "AudioGameplayLogs.h"
 #include "Modules/ModuleManager.h"
+#include "ProfilingDebugging/CsvProfiler.h"
 
 #define LOCTEXT_NAMESPACE "FAudioGameplay"
 
 DEFINE_LOG_CATEGORY(AudioGameplayLog);
 
+// Defines the "AudioGameplay" category in the CSV profiler.
+// This should only be defined here. Modules who wish to use this category should contain the line
+// 		CSV_DECLARE_CATEGORY_MODULE_EXTERN(AUDIOGAMEPLAY_API, AudioGameplay);
+//
+CSV_DEFINE_CATEGORY_MODULE(AUDIOGAMEPLAY_API, AudioGameplay, false);
+
 void FAudioGameplayModule::StartupModule()
 {
 	UE_LOG(AudioGameplayLog, Log, TEXT("%s"), ANSI_TO_TCHAR(__FUNCTION__));
