enable csv profiling for audio gameplay on non shipping builds

diff --git a/Engine/Plugins/AudioGameplay/Source/AudioGameplay/Private/AudioGameplayModule.cpp b/Engine/Plugins/AudioGameplay/Source/AudioGameplay/Private/AudioGameplayModule.cpp
--- a/Engine/Plugins/AudioGameplay/Source/AudioGameplay/Private/AudioGameplayModule.cpp
+++ b/Engine/Plugins/AudioGameplay/Source/AudioGameplay/Private/AudioGameplayModule.cpp
@@ -13,7 +13,12 @@ DEFINE_LOG_CATEGORY(AudioGameplayLog);
 // This should only be defined here. Modules who wish to use this category should contain the line
 // 		CSV_DECLARE_CATEGORY_MODULE_EXTERN(AUDIOGAMEPLAY_API, AudioGameplay);
 //
+
+#if UE_BUILD_SHIPPING
 CSV_DEFINE_CATEGORY_MODULE(AUDIOGAMEPLAY_API, AudioGameplay, false);
+#else
+CSV_DEFINE_CATEGORY_MODULE(AUDIOGAMEPLAY_API, AudioGameplay, true);
+#endif
 
 void FAudioGameplayModule::StartupModule()
 {
