Move AudioAssetUserData to AudioGameplay from AudioMetadata

diff --git a/Engine/Plugins/AudioGameplay/Source/AudioGameplay/Public/AudioAssetUserData.h b/Engine/Plugins/AudioGameplay/Source/AudioGameplay/Public/AudioAssetUserData.h
new file mode 100644
--- /dev/null
+++ b/Engine/Plugins/AudioGameplay/Source/AudioGameplay/Public/AudioAssetUserData.h
@@ -0,0 +1,21 @@
+// Copyright Epic Games, Inc. All Rights Reserved.
+
+#pragma once
+
+#include "Engine/AssetUserData.h"
+#include "GameplayTagContainer.h"
+#include "AudioAssetUserData.generated.h"
+
+/**
+ * UAudioAssetUserData - Base class for user data being attached to audio assets
+ */
+UCLASS(Blueprintable)
+class AUDIOGAMEPLAY_API UAudioAssetUserData : public UAssetUserData
+{
+	GENERATED_BODY()
+	
+public:
+
+	UPROPERTY(EditAnywhere, BlueprintReadWrite, Category = "Tags", meta = (Categories = "Audio.Metadata"))
+	FGameplayTagContainer MetadataTags;
+};
