Generic support for SoundLibraryContext tag caches.
+ Saves performance on string manipulation when looking for dynamically-determined subtags
+ All cached data is emptied upon any world travel (including PIE shutdown)
+ Extant use cases include Perspective, Surface, and HitNotify events

#rnx

diff --git a/Engine/Plugins/AudioGameplay/Source/AudioGameplay/Private/AudioGameplayTagCacheSubsystem.cpp b/Engine/Plugins/AudioGameplay/Source/AudioGameplay/Private/AudioGameplayTagCacheSubsystem.cpp
new file mode 100644
--- /dev/null
+++ b/Engine/Plugins/AudioGameplay/Source/AudioGameplay/Private/AudioGameplayTagCacheSubsystem.cpp
@@ -0,0 +1,45 @@
+// Copyright Epic Games, Inc. All Rights Reserved.
+
+#include "AudioGameplayTagCacheSubsystem.h"
+
+#include "Engine/World.h"
+
+#include UE_INLINE_GENERATED_CPP_BY_NAME(AudioGameplayTagCacheSubsystem)
+
+namespace AudioGameplayTagCacheConsoleVariables
+{
+	int32 bEnabled = 1;
+	FAutoConsoleVariableRef CVarEnabled(
+		TEXT("au.GameplayTagCache.Enabled"),
+		bEnabled,
+		TEXT("Cache any string concatenations used to build Gameplay Tags at runtime.\n0: Disable, 1: Enable (default)"),
+		ECVF_Default);
+
+} // namespace AudioGameplayTagCacheConsoleVariables
+
+bool UAudioGameplayTagCacheSubsystem::ShouldCreateSubsystem(UObject* Outer) const
+{
+	if (AudioGameplayTagCacheConsoleVariables::bEnabled == 0)
+	{
+		return false;
+	}
+
+	return Super::ShouldCreateSubsystem(Outer);
+}
+
+void UAudioGameplayTagCacheSubsystem::Deinitialize()
+{
+	GameplayTagCache.Empty();
+
+	Super::Deinitialize();
+}
+
+UAudioGameplayTagCacheSubsystem* UAudioGameplayTagCacheSubsystem::Get(const UWorld* WorldContext)
+{
+	if (WorldContext)
+	{
+		return WorldContext->GetSubsystem<UAudioGameplayTagCacheSubsystem>();
+	}
+
+	return nullptr;
+}
diff --git a/Engine/Plugins/AudioGameplay/Source/AudioGameplay/Public/AudioGameplayTagCacheSubsystem.h b/Engine/Plugins/AudioGameplay/Source/AudioGameplay/Public/AudioGameplayTagCacheSubsystem.h
new file mode 100644
--- /dev/null
+++ b/Engine/Plugins/AudioGameplay/Source/AudioGameplay/Public/AudioGameplayTagCacheSubsystem.h
@@ -0,0 +1,42 @@
+// Copyright Epic Games, Inc. All Rights Reserved.
+
+#pragma once
+
+#include "GameplayTagContainer.h"
+#include "MapAnyKey.h"
+#include "Subsystems/WorldSubsystem.h"
+
+#include "AudioGameplayTagCacheSubsystem.generated.h"
+
+using FGameplayTagMap = UE::AudioGameplay::TMapAnyKey<FGameplayTag>;
+
+/**
+ * UAudioGameplayTagCacheSubsystem - Per world subsystem used to persist gameplay tags that are expensive to construct dynamically from parts
+ */
+UCLASS()
+class AUDIOGAMEPLAY_API UAudioGameplayTagCacheSubsystem : public UWorldSubsystem
+{
+	GENERATED_BODY()
+
+public: 
+
+	UAudioGameplayTagCacheSubsystem() = default;
+	virtual ~UAudioGameplayTagCacheSubsystem() = default;
+
+	//~ Begin USubsystem interface
+	virtual bool ShouldCreateSubsystem(UObject* Outer) const override;
+	virtual void Deinitialize() override;
+	//~ End USubsystem interface
+
+	static UAudioGameplayTagCacheSubsystem* Get(const UWorld* WorldContext);
+
+	FGameplayTagMap& GetTagCache()
+	{
+		return GameplayTagCache;
+	}
+
+protected:
+
+	FGameplayTagMap GameplayTagCache;
+};
+
diff --git a/Engine/Plugins/AudioGameplay/Source/AudioGameplay/Public/MapAnyKey.h b/Engine/Plugins/AudioGameplay/Source/AudioGameplay/Public/MapAnyKey.h
new file mode 100644
--- /dev/null
+++ b/Engine/Plugins/AudioGameplay/Source/AudioGameplay/Public/MapAnyKey.h
@@ -0,0 +1,43 @@
+// Copyright Epic Games, Inc. All Rights Reserved.
+
+#include "Containers/Map.h"
+
+namespace UE::AudioGameplay
+{
+	/**
+	 * Variant of TMap that allows multiple key types within the same map instance.
+	 * Key types must be hash-able, as with a regular map.
+	 */
+	template <typename ValueType>
+	class TMapAnyKey : public TMap<uint32, ValueType>
+	{
+		using Super = TMap<uint32, ValueType>;
+
+	public:
+
+		/**
+		 * Find the value associated with a specified key, or if none exists, 
+		 * add a value using the default constructor.
+		 *
+		 * @param Key The key to search for.
+		 * @param bHadKey True if found, false if added.
+		 * @return A reference to the value associated with the specified key.
+		 */
+		template <typename AnyKeyType>
+		ValueType& FindOrAdd(AnyKeyType Key, bool& bHadKey)
+		{
+			// Calculate the hash for ANY key, using the same method a TMap<> would use for that key.
+			uint32 KeyHash = TSet<AnyKeyType>::KeyFuncsType::GetKeyHash(Key);
+
+			// Use the hashes directly as keys for the underlying TMap parent class.
+			if (ValueType* Ptr = Super::Find(KeyHash))
+			{
+				bHadKey = true;
+				return *Ptr;
+			}
+
+			bHadKey = false;
+			return Super::Add(KeyHash);
+		}
+	};
+}
