Moving Shuffle util into a more generic place for reuse on AudioGameplay side

#rnx

diff --git a/Engine/Plugins/AudioGameplay/Source/AudioGameplay/Private/AudioGameplayShuffle.cpp b/Engine/Plugins/AudioGameplay/Source/AudioGameplay/Private/AudioGameplayShuffle.cpp
new file mode 100644
--- /dev/null
+++ b/Engine/Plugins/AudioGameplay/Source/AudioGameplay/Private/AudioGameplayShuffle.cpp
@@ -0,0 +1,50 @@
+// Copyright Epic Games, Inc. All Rights Reserved.
+
+#include "AudioGameplayShuffle.h"
+
+namespace UE::Audio
+{
+	void FShuffleUtil::Initialize(int32 ArraySize)
+	{
+		ShuffleArray.SetNum(ArraySize);
+
+		for (int32 i = 0; i < ShuffleArray.Num(); i++)
+		{
+			ShuffleArray[i] = i;
+		}
+
+		Shuffle();
+	}
+
+	uint8 FShuffleUtil::GetNextIndex()
+	{
+		const int32 ShuffleArraySize = ShuffleArray.Num();
+		if (ShuffleArraySize == 0)
+		{
+			return INDEX_NONE;
+		}
+
+		if (ShuffleMarker > ShuffleArraySize)
+		{
+			Shuffle();
+		}
+
+		// Post decrement will setup an underflow to cause a shuffle on next access
+		return ShuffleArray[ShuffleMarker--];
+	}
+
+	void FShuffleUtil::Shuffle()
+	{
+		const int32 LastIndex = ShuffleArray.Num() - 1;
+		for (int32 i = 0; i <= LastIndex; ++i)
+		{
+			const int32 SwapIndex = FMath::RandRange(i, LastIndex);
+			if (i != SwapIndex)
+			{
+				ShuffleArray.Swap(i, SwapIndex);
+			}
+		}
+
+		ShuffleMarker = LastIndex;
+	}
+} // namespace UE::Audio
diff --git a/Engine/Plugins/AudioGameplay/Source/AudioGameplay/Public/AudioGameplayShuffle.h b/Engine/Plugins/AudioGameplay/Source/AudioGameplay/Public/AudioGameplayShuffle.h
new file mode 100644
--- /dev/null
+++ b/Engine/Plugins/AudioGameplay/Source/AudioGameplay/Public/AudioGameplayShuffle.h
@@ -0,0 +1,29 @@
+// Copyright Epic Games, Inc. All Rights Reserved.
+
+#pragma once
+
+#include "Containers/Array.h"
+
+namespace UE::Audio
+{
+	/**
+	 * FShuffleUtil - Utility to provide a fast random index in an array without replacement
+	 *   ie, all elements will be returned from GetNextIndex exactly once before repeating
+	 */
+	struct AUDIOGAMEPLAY_API FShuffleUtil
+	{
+		/** Initializes the Shuffle Array */
+		void Initialize(int32 ArraySize);
+
+		/** Returns next valid index, or INDEX_NONE if not initialized. Shuffles array if needed */
+		uint8 GetNextIndex();
+
+	protected:
+
+		uint8 ShuffleMarker = 0;
+		TArray<uint8> ShuffleArray;
+
+		/** Shuffles the internal array */
+		void Shuffle();
+	};
+} // namespace UE::Audio
