[Audio] Add BP utility functions to AudioAssetUserData for interacting with audio asset tags

#rnx

diff --git a/Engine/Plugins/AudioGameplay/Source/AudioGameplay/Private/AudioAssetUserData.cpp b/Engine/Plugins/AudioGameplay/Source/AudioGameplay/Private/AudioAssetUserData.cpp
new file mode 100644
--- /dev/null
+++ b/Engine/Plugins/AudioGameplay/Source/AudioGameplay/Private/AudioAssetUserData.cpp
@@ -0,0 +1,57 @@
+// Copyright Epic Games, Inc. All Rights Reserved.
+#include "AudioAssetUserData.h"
+
+#include "Sound/SoundBase.h"
+#include "Sound/SoundClass.h"
+
+#include UE_INLINE_GENERATED_CPP_BY_NAME(AudioAssetUserData)
+
+FGameplayTagContainer UAudioAssetUserData::GetAllTags(USoundBase* InSound)
+{
+	FGameplayTagContainer AllTags;
+	if (InSound == nullptr)
+	{
+		return AllTags;
+	}
+	
+	// Check Sound directly
+	if (UAudioAssetUserData* SoundData = Get(InSound))
+	{
+		AllTags.AppendTags(SoundData->MetadataTags);
+	}
+	
+	// Check SoundClass
+	if (UAudioAssetUserData* ClassData = Get(InSound->GetSoundClass()))
+	{
+		AllTags.AppendTags(ClassData->MetadataTags);
+	}
+	
+	return AllTags;
+}
+
+bool UAudioAssetUserData::HasTag(USoundBase* InSound, FGameplayTag InTag, bool bExactMatch)
+{
+	FGameplayTagContainer AllTags = GetAllTags(InSound);
+	if (bExactMatch)
+	{
+		return AllTags.HasTagExact(InTag);
+	}
+
+	return AllTags.HasTag(InTag);
+}
+
+FGameplayTagContainer UAudioAssetUserData::GetFilteredTags(USoundBase* InSound, FGameplayTag InTag)
+{
+	const FGameplayTagContainer AllTags = GetAllTags(InSound);
+	return AllTags.Filter(FGameplayTagContainer(InTag));
+}
+
+UAudioAssetUserData* UAudioAssetUserData::Get(UObject* InObject)
+{
+	if (IInterface_AssetUserData* Interface = Cast<IInterface_AssetUserData>(InObject))
+	{
+		return Cast<UAudioAssetUserData>(Interface->GetAssetUserDataOfClass(StaticClass()));
+	}
+	
+	return nullptr;
+}
diff --git a/Engine/Plugins/AudioGameplay/Source/AudioGameplay/Public/AudioAssetUserData.h b/Engine/Plugins/AudioGameplay/Source/AudioGameplay/Public/AudioAssetUserData.h
--- a/Engine/Plugins/AudioGameplay/Source/AudioGameplay/Public/AudioAssetUserData.h
+++ b/Engine/Plugins/AudioGameplay/Source/AudioGameplay/Public/AudioAssetUserData.h
@@ -6,6 +6,8 @@
 #include "GameplayTagContainer.h"
 #include "AudioAssetUserData.generated.h"
 
+#define UE_API AUDIOGAMEPLAY_API
+
 /**
  * UAudioAssetUserData - Base class for user data being attached to audio assets
  */
@@ -16,6 +18,23 @@ class UAudioAssetUserData : public UAssetUserData
 	
 public:
 
-	UPROPERTY(EditAnywhere, BlueprintReadWrite, Category = "Tags", meta = (Categories = "Audio.Metadata"))
+	UPROPERTY(EditAnywhere, BlueprintReadWrite, Category = "Tags")
 	FGameplayTagContainer MetadataTags;
+
+	// Gather all AudioAssetUserData tags for this sound and its sound class. 
+	UFUNCTION(BlueprintPure, Category = "Audio AssetUserData Tags")
+	static UE_API FGameplayTagContainer GetAllTags(USoundBase* InSound);
+
+	// Whether the sound or its sound class have the supplied (possibly exact) tag .
+	UFUNCTION(BlueprintPure, Category = "Audio AssetUserData Tags")
+	static UE_API bool HasTag(USoundBase* InSound, FGameplayTag InTag, bool bExactMatch = false);
+
+	// Gather all AudioAssetUserData tags for this sound and its sound class, that match the supplied tag.
+	UFUNCTION(BlueprintPure, Category = "Audio AssetUserData Tags")
+	static UE_API FGameplayTagContainer GetFilteredTags(USoundBase* InSound, FGameplayTag InTag);
+
+	// If the supplied object has a UAudioAssetUserData asset user data attached, gets it. Returns nullptr otherwise.
+	static UE_API UAudioAssetUserData* Get(UObject* InObject);
 };
+
+#undef UE_API
