Added missing include.

#rnx

diff --git a/Engine/Plugins/AudioGameplay/Source/AudioGameplay/Private/AudioAssetUserData.cpp b/Engine/Plugins/AudioGameplay/Source/AudioGameplay/Private/AudioAssetUserData.cpp
--- a/Engine/Plugins/AudioGameplay/Source/AudioGameplay/Private/AudioAssetUserData.cpp
+++ b/Engine/Plugins/AudioGameplay/Source/AudioGameplay/Private/AudioAssetUserData.cpp
@@ -3,6 +3,7 @@
 
 #include "Sound/SoundBase.h"
 #include "Sound/SoundClass.h"
+#include "Templates/SubclassOf.h"
 
 #include UE_INLINE_GENERATED_CPP_BY_NAME(AudioAssetUserData)
 
